Add a cancel button to cancel a sync

TODO add a pop up to confirm yes/no you really want to cancel such a thing

diff --git a/Engine/Source/Programs/SlateUGS/Private/UGSTab.cpp b/Engine/Source/Programs/SlateUGS/Private/UGSTab.cpp
--- a/Engine/Source/Programs/SlateUGS/Private/UGSTab.cpp
+++ b/Engine/Source/Programs/SlateUGS/Private/UGSTab.cpp
@@ -26,7 +26,8 @@ UGSTab::UGSTab() : TabArgs(nullptr, FTabId()),
 				   TabWidget(SNew(SDockTab)),
 				   EmptyTabView(SNew(SEmptyTab).Tab(this)),
 				   GameSyncTabView(SNew(SGameSyncTab).Tab(this)),
-				   bHasQueuedMessages(false)
+				   bHasQueuedMessages(false),
+				   bNeedUpdateGameTabBuildList(false)
 {
 	Initialize(nullptr);
 }
@@ -202,6 +203,14 @@ bool UGSTab::OnWorkspaceChosen(const FString& Project)
 	return false;
 }
 
+void UGSTab::CancelSync()
+{
+	if (Workspace->IsBusy())
+	{
+		Workspace->CancelUpdate();
+	}
+}
+
 void UGSTab::OnSyncChangelist(int Changelist)
 {
 	// Options on what to do with workspace when updating it
@@ -512,23 +521,28 @@ void UGSTab::OnWorkspaceSyncComplete(TSharedRef<UGSCore::FWorkspaceUpdateContext
 {
 	FScopeLock Lock(&CriticalSection);
 
-	WorkspaceSettings->CurrentChangeNumber     = Workspace->GetCurrentChangeNumber();
-	WorkspaceSettings->LastBuiltChangeNumber   = Workspace->GetLastBuiltChangeNumber();
-	WorkspaceSettings->LastSyncResult          = SyncResult;
-	WorkspaceSettings->LastSyncResultMessage   = StatusMessage;
-	WorkspaceSettings->LastSyncTime            = WorkspaceContext->StartTime;
+	if (SyncResult == UGSCore::EWorkspaceUpdateResult::Success)
+	{
+		WorkspaceSettings->CurrentChangeNumber   = Workspace->GetCurrentChangeNumber();
+		WorkspaceSettings->LastBuiltChangeNumber = Workspace->GetLastBuiltChangeNumber();
+
+		// Queue up setting the changelist text on the main thread
+		QueueMessageForMainThread([this] {
+			GameSyncTabView->SetChangelistText(WorkspaceSettings->CurrentChangeNumber);
+		});
+
+		// TODO hacky, but allows us to highlight which CL we have synced to in the list
+		bNeedUpdateGameTabBuildList = true;
+	}
+
+	WorkspaceSettings->LastSyncResult        = SyncResult;
+	WorkspaceSettings->LastSyncResultMessage = StatusMessage;
+	WorkspaceSettings->LastSyncTime          = WorkspaceContext->StartTime;
+
 	// TODO check this is valid, may be off
 	WorkspaceSettings->LastSyncDurationSeconds = (FDateTime::UtcNow() - WorkspaceContext->StartTime).GetSeconds();
 
-	// Queue up setting the changelist text on the main thread
-	QueueMessageForMainThread([this] {
-		GameSyncTabView->SetChangelistText(WorkspaceSettings->CurrentChangeNumber);
-	});
-
 	UserSettings->Save();
-
-	// TODO hacky, but allows us to highlight which CL we have synced to in the list
-	bNeedUpdateGameTabBuildList = true;
 }
 
 void UGSTab::SetupWorkspace()
diff --git a/Engine/Source/Programs/SlateUGS/Private/UGSTab.h b/Engine/Source/Programs/SlateUGS/Private/UGSTab.h
--- a/Engine/Source/Programs/SlateUGS/Private/UGSTab.h
+++ b/Engine/Source/Programs/SlateUGS/Private/UGSTab.h
@@ -66,6 +66,8 @@ public:
 	UGSTabManager* GetTabManager();
 	TSharedPtr<UGSCore::FUserSettings> GetUserSettings() const;
 
+	void CancelSync();
+
 private:
 	void OnWorkspaceSyncComplete(
 		TSharedRef<UGSCore::FWorkspaceUpdateContext, ESPMode::ThreadSafe> WorkspaceContext,
@@ -108,5 +110,5 @@ private:
 	TArray<TFunction<void()>> MessageQueue;
 	std::atomic<bool> bHasQueuedMessages;
 
-	bool bNeedUpdateGameTabBuildList = false;
+	std::atomic<bool> bNeedUpdateGameTabBuildList;
 };
diff --git a/Engine/Source/Programs/SlateUGS/Private/Widgets/SGameSyncTab.cpp b/Engine/Source/Programs/SlateUGS/Private/Widgets/SGameSyncTab.cpp
--- a/Engine/Source/Programs/SlateUGS/Private/Widgets/SGameSyncTab.cpp
+++ b/Engine/Source/Programs/SlateUGS/Private/Widgets/SGameSyncTab.cpp
@@ -182,7 +182,6 @@ void SGameSyncTab::Construct(const FArguments& InArgs)
 		.Padding(10.0f, 5.0f)
 		[
 			SNew(SBorder)
-			.IsEnabled_Lambda([this] { return !Tab->IsSyncing(); })
 			.BorderImage(FAppStyle::Get().GetBrush("Brushes.Panel"))
 			[
 				SNew(SHorizontalBox)
@@ -201,6 +200,16 @@ void SGameSyncTab::Construct(const FArguments& InArgs)
 						[
 							MakeSyncButtonDropdown()
 						]
+						.IsEnabled_Lambda([this] { return !Tab->IsSyncing(); })
+					]
+					+SHorizontalBox::Slot()
+					.AutoWidth()
+					[
+						SNew(SSimpleButton)
+						.Text(LOCTEXT("CancelSync", "Cancel"))
+						.Icon(FAppStyle::Get().GetBrush("Icons.X"))
+						.OnClicked_Lambda([this] { Tab->CancelSync(); return FReply::Handled(); })
+						.IsEnabled_Lambda([this] { return Tab->IsSyncing(); })
 					]
 					+SHorizontalBox::Slot()
 					.AutoWidth()
@@ -214,6 +223,7 @@ void SGameSyncTab::Construct(const FArguments& InArgs)
 					[
 						SNew(SSimpleButton)
 						.Text(LOCTEXT("Build", "Build"))
+						.IsEnabled_Lambda([this] { return !Tab->IsSyncing(); })
 						.IsEnabled(false) // Todo: enable after adding this functionality
 					]
 					+SHorizontalBox::Slot()
@@ -223,6 +233,7 @@ void SGameSyncTab::Construct(const FArguments& InArgs)
 						.Text(LOCTEXT("RunUnrealEditor", "Run Unreal Editor"))
 						.Icon(FAppStyle::Get().GetBrush("Icons.Launch"))
 						.HasDownArrow(true)
+						.IsEnabled_Lambda([this] { return !Tab->IsSyncing(); })
 						.IsEnabled(false) // Todo: enable after adding this functionality
 					]
 					+SHorizontalBox::Slot()
@@ -231,6 +242,7 @@ void SGameSyncTab::Construct(const FArguments& InArgs)
 						SNew(SSimpleButton)
 						.Text(LOCTEXT("OpenSolution", "OpenSolution"))
 						.Icon(FAppStyle::Get().GetBrush("MainFrame.OpenVisualStudio")) // Todo: shouldn't use this icon (repurposing, also could use other IDEs)
+						.IsEnabled_Lambda([this] { return !Tab->IsSyncing(); })
 						.IsEnabled(false) // Todo: enable after adding this functionality
 					]
 				]
@@ -267,6 +279,7 @@ void SGameSyncTab::Construct(const FArguments& InArgs)
 						SNew(SSimpleButton)
 						.Text(LOCTEXT("CleanSolution", "Clean Solution"))
 						.Icon(FAppStyle::Get().GetBrush("GraphEditor.Clean")) // Todo: shouldn't use this icon (repurposing)
+						.IsEnabled_Lambda([this] { return !Tab->IsSyncing(); })
 						.IsEnabled(false) // Todo: enable after adding this functionality
 					]
 					+SHorizontalBox::Slot()
diff --git a/Engine/Source/Programs/UGSCore/Perforce.cpp b/Engine/Source/Programs/UGSCore/Perforce.cpp
--- a/Engine/Source/Programs/UGSCore/Perforce.cpp
+++ b/Engine/Source/Programs/UGSCore/Perforce.cpp
@@ -84,7 +84,7 @@ struct FPerforceExe
 		return ExecutablePath;
 	}
 
-	inline int InnerRunCommand(const FString& CommandLine, TArray<FString>& OutLines)
+	inline int InnerRunCommand(const FString& CommandLine, TArray<FString>& OutLines, FEvent* AbortEvent)
 	{
 		uint32 ProcId;
 		void* ReadPipe = nullptr;
@@ -99,7 +99,12 @@ struct FPerforceExe
 		{
 			P4Output += LatestOutput;
 			LatestOutput = FPlatformProcess::ReadPipe(ReadPipe);
-			FPlatformProcess::Sleep(0);
+
+			if (AbortEvent->Wait(FTimespan::Zero()))
+			{
+				FPlatformProcess::TerminateProc(P4Proc);
+				return -1;
+			}
 		}
 
 		FPlatformProcess::ClosePipe(ReadPipe, WritePipe);
@@ -117,11 +122,11 @@ struct FPerforceExe
 		return -1;
 	}
 
-	int RunCommand(const FString& CommandLine, TArray<FString>& OutLines)
+	int RunCommand(const FString& CommandLine, TArray<FString>& OutLines, FEvent* AbortEvent)
 	{
 		if (bValidP4)
 		{
-			return InnerRunCommand(CommandLine, OutLines);
+			return InnerRunCommand(CommandLine, OutLines, AbortEvent);
 		}
 
 		return -1;
@@ -1123,7 +1128,7 @@ bool FPerforceConnection::RunCommand(const FString& CommandLine, EPerforceOutput
 
 	TArray<FString> RawOutputLines;
 
-	int ExitCode = GPerforceExe.RunCommand(FullCommandLine, RawOutputLines);
+	int ExitCode = GPerforceExe.RunCommand(FullCommandLine, RawOutputLines, AbortEvent);
 	if (ExitCode != 0 && !EnumHasAnyFlags(Options, ECommandOptions::IgnoreExitCode))
 	{
 		return false;
@@ -1153,7 +1158,7 @@ bool FPerforceConnection::RunCommand(const FString& CommandLine, const TCHAR* In
 
 	TArray<FString> RawOutputLines;
 
-	int ExitCode = GPerforceExe.RunCommand(FullCommandLine, RawOutputLines);
+	int ExitCode = GPerforceExe.RunCommand(FullCommandLine, RawOutputLines, AbortEvent);
 
 	// TODO check this vs the old way as things *may* have changed. Before they were handling each line as it was coming
 	// back from the process. here we just spin + collect all the output. Then try to process all the output, then if the
@@ -1408,8 +1413,8 @@ bool FPerforceConnection::RunCommandWithBinaryOutput(const FString& CommandLine,
 
 	uint32 ProcId;
 	void* ReadPipe = nullptr;
-    void* WritePipe = nullptr;
-    FPlatformProcess::CreatePipe(ReadPipe, WritePipe);
+	void* WritePipe = nullptr;
+	FPlatformProcess::CreatePipe(ReadPipe, WritePipe);
 
 	FProcHandle P4Proc = FPlatformProcess::CreateProc(*GPerforceExe.GetPerforceExe(), *FullCommandLine, false, true, true, &ProcId, 0, nullptr, WritePipe, ReadPipe);
 
@@ -1425,6 +1430,7 @@ bool FPerforceConnection::RunCommandWithBinaryOutput(const FString& CommandLine,
 		// Check that we don't have an abort event
 		if (AbortEvent->Wait(FTimespan::Zero()))
 		{
+			FPlatformProcess::TerminateProc(P4Proc);
 			return false;
 			//throw FAbortException();
 		}
@@ -1446,6 +1452,7 @@ bool FPerforceConnection::RunCommandWithBinaryOutput(const FString& CommandLine,
 			}
 			else if(AbortEvent->Wait(FTimespan::FromMilliseconds(50)))
 			{
+				FPlatformProcess::TerminateProc(P4Proc);
 				return false;
 				//throw FAbortException();
 			}
@@ -1480,7 +1487,7 @@ bool FPerforceConnection::RunCommandWithBinaryOutput(const FString& CommandLine,
 		}
 	}
 
-    FPlatformProcess::ClosePipe(ReadPipe, WritePipe);
+	FPlatformProcess::ClosePipe(ReadPipe, WritePipe);
 
 	int ExitCode = -1;
 
diff --git a/Engine/Source/Programs/UGSCore/Utility.cpp b/Engine/Source/Programs/UGSCore/Utility.cpp
--- a/Engine/Source/Programs/UGSCore/Utility.cpp
+++ b/Engine/Source/Programs/UGSCore/Utility.cpp
@@ -115,21 +115,26 @@ int FUtility::ExecuteProcess(const TCHAR* FileName, const TCHAR* CommandLine, co
 {
 	uint32 ProcId;
 	void* ReadPipe = nullptr;
-    void* WritePipe = nullptr;
-    FPlatformProcess::CreatePipe(ReadPipe, WritePipe);
+	void* WritePipe = nullptr;
+	FPlatformProcess::CreatePipe(ReadPipe, WritePipe);
 
 	FProcHandle Proc = FPlatformProcess::CreateProc(FileName, CommandLine, false, true, true, &ProcId, 0, nullptr, WritePipe, ReadPipe);
 
 	FString Output;
-    FString LatestOutput = FPlatformProcess::ReadPipe(ReadPipe);
-    while (FPlatformProcess::IsProcRunning(Proc) || !LatestOutput.IsEmpty())
-    {
-        Output += LatestOutput;
-        LatestOutput = FPlatformProcess::ReadPipe(ReadPipe);
-		FPlatformProcess::Sleep(0);
-    }
+	FString LatestOutput = FPlatformProcess::ReadPipe(ReadPipe);
+	while (FPlatformProcess::IsProcRunning(Proc) || !LatestOutput.IsEmpty())
+	{
+		Output += LatestOutput;
+		LatestOutput = FPlatformProcess::ReadPipe(ReadPipe);
+
+		if (AbortEvent->Wait(FTimespan::Zero()))
+		{
+			FPlatformProcess::TerminateProc(Proc);
+			return -1;
+		}
+	}
 
-    FPlatformProcess::ClosePipe(ReadPipe, WritePipe);
+	FPlatformProcess::ClosePipe(ReadPipe, WritePipe);
 
 	// TODO worried this may not always work for say p4 servers sending out only \n on Windows host
 	Output.ParseIntoArray(OutLines, LINE_TERMINATOR);
@@ -149,14 +154,14 @@ int FUtility::ExecuteProcess(const TCHAR* FileName, const TCHAR* CommandLine, co
 {
 	uint32 ProcId;
 	void* ReadPipe = nullptr;
-    void* WritePipe = nullptr;
-    FPlatformProcess::CreatePipe(ReadPipe, WritePipe);
+	void* WritePipe = nullptr;
+	FPlatformProcess::CreatePipe(ReadPipe, WritePipe);
 
 	FProcHandle Proc = FPlatformProcess::CreateProc(FileName, CommandLine, false, true, true, &ProcId, 0, nullptr, WritePipe, ReadPipe);
 
-    FString LatestOutput = FPlatformProcess::ReadPipe(ReadPipe);
-    while (FPlatformProcess::IsProcRunning(Proc) || !LatestOutput.IsEmpty())
-    {
+	FString LatestOutput = FPlatformProcess::ReadPipe(ReadPipe);
+	while (FPlatformProcess::IsProcRunning(Proc) || !LatestOutput.IsEmpty())
+	{
 		// Check if we end up having a perfect data blob, with a perfect amount of lines
 		bool bEndsWithLineTerminator = LatestOutput.EndsWith(LINE_TERMINATOR);
 
@@ -180,10 +185,15 @@ int FUtility::ExecuteProcess(const TCHAR* FileName, const TCHAR* CommandLine, co
 		}
 
 		LatestOutput = LeftOverLine + FPlatformProcess::ReadPipe(ReadPipe);
-		FPlatformProcess::Sleep(0);
-    }
 
-    FPlatformProcess::ClosePipe(ReadPipe, WritePipe);
+		if (AbortEvent->Wait(FTimespan::Zero()))
+		{
+			FPlatformProcess::TerminateProc(Proc);
+			return -1;
+		}
+	}
+
+	FPlatformProcess::ClosePipe(ReadPipe, WritePipe);
 
 	int ExitCode = -1;
 	bool GotReturnCode = FPlatformProcess::GetProcReturnCode(Proc, &ExitCode);
