SlateUGS: Fix parsing of precompiled editor/binaries

diff --git a/Engine/Source/Programs/UGSCore/Perforce.cpp b/Engine/Source/Programs/UGSCore/Perforce.cpp
--- a/Engine/Source/Programs/UGSCore/Perforce.cpp
+++ b/Engine/Source/Programs/UGSCore/Perforce.cpp
@@ -726,7 +726,7 @@ bool FPerforceConnection::FindFileChanges(const FString& FilePath, int MaxResult
 
 bool FPerforceConnection::TryParseFileChangeSummary(const TArray<FString>& Lines, int& LineIdx, FPerforceFileChangeSummary& OutChange) const
 {
-	TArray<FString> Tokens = Split(Lines[LineIdx], TEXT(" "), true);
+	TArray<FString> Tokens = Split(Lines[LineIdx].TrimStartAndEnd(), TEXT(" "), true);
 	if(Tokens.Num() != 10 || !Tokens[0].StartsWith(TEXT("#")) || Tokens[1] != TEXT("change") || Tokens[4] != TEXT("on") || Tokens[7] != TEXT("by"))
 	{
 		return false;
@@ -741,7 +741,7 @@ bool FPerforceConnection::TryParseFileChangeSummary(const TArray<FString>& Lines
 	}
 
 	int UserClientIdx;
-	if(Tokens[8].FindChar('@', UserClientIdx))
+	if(!Tokens[8].FindChar('@', UserClientIdx))
 	{
 		return false;
 	}
diff --git a/Engine/Source/Programs/UGSCore/PerforceMonitor.cpp b/Engine/Source/Programs/UGSCore/PerforceMonitor.cpp
--- a/Engine/Source/Programs/UGSCore/PerforceMonitor.cpp
+++ b/Engine/Source/Programs/UGSCore/PerforceMonitor.cpp
@@ -489,7 +489,7 @@ bool FPerforceMonitor::UpdateZippedBinaries()
 					int OriginalChangeNumber;
 					if(FUtility::TryParse(*Tokens[1].Mid(0, Tokens[1].Len() - 1), OriginalChangeNumber) && !NewChangeNumberToZippedBinaries.Contains(OriginalChangeNumber))
 					{
-						NewChangeNumberToZippedBinaries[OriginalChangeNumber] = FString::Printf(TEXT("%s#%d"), *ZippedBinariesPath, ZipChange.Revision);
+						NewChangeNumberToZippedBinaries.Add(OriginalChangeNumber, FString::Printf(TEXT("%s#%d"), *ZippedBinariesPath, ZipChange.Revision));
 					}
 				}
 			}
