- Fixed MacPlatformProcess child write pips to block so it doesn't blow the buffer for spammy comands
- Fixed Mac Perforce exe location discovery code for when there are multiple p4's found

diff --git a/Engine/Source/Programs/UGSCore/Perforce.cpp b/Engine/Source/Programs/UGSCore/Perforce.cpp
--- a/Engine/Source/Programs/UGSCore/Perforce.cpp
+++ b/Engine/Source/Programs/UGSCore/Perforce.cpp
@@ -67,6 +67,8 @@ struct FPerforceExe
 
 #if PLATFORM_MAC
 		FPlatformProcess::ExecProcess(TEXT("/usr/bin/mdfind"), TEXT("\"kMDItemFSName = 'p4' && kMDItemContentType = 'public.unix-executable'\""), &ReturnCode, &OutResults, &OutErrors);
+		// in case it returns multiple p4's, pick the first one
+		OutResults = Split(OutResults, LINE_TERMINATOR)[0];
 #elif PLATFORM_LINUX
 		FPlatformProcess::ExecProcess(TEXT("which"), TEXT("p4"), &ReturnCode, &OutResults, &OutErrors);
 #endif // PLATFORM_MAC
diff --git a/Engine/Source/Runtime/Core/Private/Mac/MacPlatformProcess.cpp b/Engine/Source/Runtime/Core/Private/Mac/MacPlatformProcess.cpp
--- a/Engine/Source/Runtime/Core/Private/Mac/MacPlatformProcess.cpp
+++ b/Engine/Source/Runtime/Core/Private/Mac/MacPlatformProcess.cpp
@@ -1317,7 +1317,7 @@ bool FMacPlatformProcess::CreatePipe( void*& ReadPipe, void*& WritePipe, bool bW
 	pipe(pipefd);
 
 	fcntl(pipefd[0], F_SETFL, O_NONBLOCK);
-	fcntl(pipefd[1], F_SETFL, O_NONBLOCK);
+	//fcntl(pipefd[1], F_SETFL, O_NONBLOCK);
 
 	// create an NSFileHandle from the descriptor
 	ReadPipe = [[NSFileHandle alloc] initWithFileDescriptor: pipefd[0]];
