When setting up a workspace, we now only search for clients (workspaces) associated with the current user. This is makes the modal loading dialog noticeably faster and this is how it's done in current/old UGS

diff --git a/Engine/Source/Programs/UGSCore/DetectProjectSettingsTask.cpp b/Engine/Source/Programs/UGSCore/DetectProjectSettingsTask.cpp
--- a/Engine/Source/Programs/UGSCore/DetectProjectSettingsTask.cpp
+++ b/Engine/Source/Programs/UGSCore/DetectProjectSettingsTask.cpp
@@ -69,7 +69,7 @@ TSharedRef<FModalTaskResult> FDetectProjectSettingsTask::RunInternal(FEvent* Abo
 	// Find all the clients on this machine
 	Log->Logf(TEXT("Enumerating clients on local machine..."));
 	TArray<FPerforceClientRecord> Clients;
-	if(!Perforce->FindClients(Clients, AbortEvent, Log.Get()))
+	if(!Perforce->FindClients(Clients, PerforceInfo->UserName, AbortEvent, Log.Get()))
 	{
 		return FModalTaskResult::Failure(LOCTEXT("NoClientsForHost", "Couldn't find any clients for this host."));
 	}
diff --git a/Engine/Source/Programs/UGSCore/Perforce.cpp b/Engine/Source/Programs/UGSCore/Perforce.cpp
--- a/Engine/Source/Programs/UGSCore/Perforce.cpp
+++ b/Engine/Source/Programs/UGSCore/Perforce.cpp
@@ -507,12 +507,11 @@ bool FPerforceConnection::FindClients(TArray<FPerforceClientRecord>& Clients, FE
 	return RunCommand(TEXT("clients"), Clients, ECommandOptions::NoClient, AbortEvent, Log);
 }
 
-bool FPerforceConnection::FindClients(const FString& ForUserName, TArray<FString>& ClientNames, FEvent* AbortEvent, FOutputDevice& Log) const
+bool FPerforceConnection::FindClients(TArray<FPerforceClientRecord>& Clients, const FString& ForUserName, FEvent* AbortEvent, FOutputDevice& Log) const
 {
 	TArray<FString> Lines;
-	if(!RunCommand(FString::Printf(TEXT("clients -u%s"), *ForUserName), Lines, ECommandOptions::None, AbortEvent, Log))
+	if(!RunCommand(FString::Printf(TEXT("clients -u%s"), *ForUserName), Clients, ECommandOptions::NoClient, AbortEvent, Log))
 	{
-		ClientNames.Empty();
 		return false;
 	}
 
@@ -523,10 +522,6 @@ bool FPerforceConnection::FindClients(const FString& ForUserName, TArray<FString
 		{
 			Log.Logf(TEXT("Couldn't parse client from line '%s'"), *Line);
 		}
-		else
-		{
-			ClientNames.Add(Tokens[1]);
-		}
 	}
 	return true;
 }
diff --git a/Engine/Source/Programs/UGSCore/Perforce.h b/Engine/Source/Programs/UGSCore/Perforce.h
--- a/Engine/Source/Programs/UGSCore/Perforce.h
+++ b/Engine/Source/Programs/UGSCore/Perforce.h
@@ -206,7 +206,7 @@ public:
 	bool Info(TSharedPtr<FPerforceInfoRecord>& OutInfo, FEvent* AbortEvent, FOutputDevice& Log) const;
 	bool GetSetting(const FString& Name, FString& Value, FEvent* AbortEvent, FOutputDevice& Log) const;
 	bool FindClients(TArray<FPerforceClientRecord>& Clients, FEvent* AbortEvent, FOutputDevice& Log) const;
-	bool FindClients(const FString& ForUserName, TArray<FString>& ClientNames, FEvent* AbortEvent, FOutputDevice& Log) const;
+	bool FindClients(TArray<FPerforceClientRecord>& Clients, const FString& ForUserName, FEvent* AbortEvent, FOutputDevice& Log) const;
 	bool TryGetClientSpec(const FString& ClientName, TSharedPtr<FPerforceSpec>& OutSpec, FEvent* AbortEvent, FOutputDevice& Log) const;
 	bool TryGetStreamSpec(const FString& StreamName, TSharedPtr<FPerforceSpec>& OutSpec, FEvent* AbortEvent, FOutputDevice& Log) const;
 	bool FindFiles(const FString& Filter, TArray<FPerforceFileRecord>& FileRecords, FEvent* AbortEvent, FOutputDevice& Log) const;
