Allows creation of new workspaces through SlateUGS

diff --git a/Engine/Source/Programs/SlateUGS/Private/UGSTab.cpp b/Engine/Source/Programs/SlateUGS/Private/UGSTab.cpp
--- a/Engine/Source/Programs/SlateUGS/Private/UGSTab.cpp
+++ b/Engine/Source/Programs/SlateUGS/Private/UGSTab.cpp
@@ -4,6 +4,7 @@
 
 #include "UGSLog.h"
 #include "ChangeInfo.h"
+#include "CreateClientTask.h"
 #include "FindStreamsTask.h"
 #include "UGSTabManager.h"
 
@@ -220,8 +221,6 @@ TArray<FString> UGSTab::GetAllStreamNames() const
 		return TArray<FString>();
 	}
 
-	Result.Sort();
-
 	return Result;
 }
 
@@ -388,6 +387,28 @@ void UGSTab::OnOpenEditor()
 	EditorProcessHandle = FPlatformProcess::CreateProc(*EditorPath, *ProjectFileName, true, false, false, nullptr, 0, nullptr, nullptr, nullptr);
 }
 
+void UGSTab::OnCreateWorkspace(const FString& WorkspaceName, const FString& Stream, const FString& RootDirectory) const
+{
+	const TSharedRef<FLineWriter> CreateWorkspaceLog = MakeShared<FLineWriter>();
+	FEvent* AbortEvent = FPlatformProcess::GetSynchEventFromPool(true);
+
+	TSharedPtr<UGSCore::FPerforceInfoRecord> PerforceInfo;
+	PerforceClient->Info(PerforceInfo, AbortEvent, *CreateWorkspaceLog);
+
+	TMap<FString, FString> Tags;
+	Tags.Add(TEXT("client"), WorkspaceName);
+	Tags.Add(TEXT("Owner"), PerforceInfo->UserName);
+	Tags.Add(TEXT("Host"), PerforceInfo->HostName);
+	Tags.Add(TEXT("Root"), RootDirectory);
+	const UGSCore::FPerforceClientRecord ClientRecord(Tags);
+
+	const TSharedRef<UGSCore::FModalTaskResult> TaskResult = ExecuteModalTask(
+		FSlateApplication::Get().GetActiveModalWindow(),
+		MakeShared<UGSCore::FCreateClientTask>(PerforceClient.ToSharedRef(), MakeShared<FLineWriter>(), ClientRecord, Stream),
+		LOCTEXT("OpeningProjectTitle", "Creating Project"),
+		LOCTEXT("OpeningProjectCaption", "Creating project, please wait..."));
+}
+
 void UGSTab::QueueMessageForMainThread(TFunction<void()> Function)
 {
 	FScopeLock Lock(&CriticalSection);
diff --git a/Engine/Source/Programs/SlateUGS/Private/UGSTab.h b/Engine/Source/Programs/SlateUGS/Private/UGSTab.h
--- a/Engine/Source/Programs/SlateUGS/Private/UGSTab.h
+++ b/Engine/Source/Programs/SlateUGS/Private/UGSTab.h
@@ -62,6 +62,7 @@ public:
 	void OnBuildWorkspace();
 	void OnOpenExplorer();
 	void OnOpenEditor();
+	void OnCreateWorkspace(const FString& WorkspaceName, const FString& Stream, const FString& RootDirectory) const;
 
 	// Accessors
 	bool IsSyncing() const;
diff --git a/Engine/Source/Programs/SlateUGS/Private/Widgets/SNewWorkspaceWindow.cpp b/Engine/Source/Programs/SlateUGS/Private/Widgets/SNewWorkspaceWindow.cpp
--- a/Engine/Source/Programs/SlateUGS/Private/Widgets/SNewWorkspaceWindow.cpp
+++ b/Engine/Source/Programs/SlateUGS/Private/Widgets/SNewWorkspaceWindow.cpp
@@ -2,6 +2,7 @@
 
 #include "SNewWorkspaceWindow.h"
 
+#include "DesktopPlatformModule.h"
 #include "SlateUGSStyle.h"
 #include "Framework/Application/SlateApplication.h"
 
@@ -13,9 +14,10 @@
 
 #define LOCTEXT_NAMESPACE "UGSNewWorkspaceWindow"
 
-void SNewWorkspaceWindow::Construct(const FArguments& InArgs, UGSTab* InTab)
+void SNewWorkspaceWindow::Construct(const FArguments& InArgs, TSharedPtr<SWorkspaceWindow> InParent, UGSTab* InTab)
 {
 	Tab = InTab;
+	Parent = InParent;
 
 	SWindow::Construct(SWindow::FArguments()
 	.Title(LOCTEXT("WindowTitle", "New Workspace"))
@@ -100,7 +102,7 @@ void SNewWorkspaceWindow::Construct(const FArguments& InArgs, UGSTab* InTab)
 					]
 					+SVerticalBox::Slot()
 					[
-						SAssignNew(FileNameTextBox, SEditableTextBox)
+						SAssignNew(WorkspaceNameTextBox, SEditableTextBox)
 					]
 				]
 			]
@@ -119,6 +121,7 @@ void SNewWorkspaceWindow::Construct(const FArguments& InArgs, UGSTab* InTab)
 						SNew(SPrimaryButton)
 						.Text(LOCTEXT("CreateButtonText", "Create"))
 						.OnClicked(this, &SNewWorkspaceWindow::OnCreateClicked)
+						.IsEnabled(this, &SNewWorkspaceWindow::IsCreateButtonEnabled)
 					]
 					+SHorizontalBox::Slot()
 					[
@@ -144,11 +147,33 @@ FReply SNewWorkspaceWindow::OnBrowseStreamClicked()
 
 FReply SNewWorkspaceWindow::OnBrowseRootDirectoryClicked()
 {
+	bool bFolderSelected = FDesktopPlatformModule::Get()->OpenDirectoryDialog(
+		FSlateApplication::Get().FindBestParentWindowHandleForDialogs(AsShared()),
+		LOCTEXT("OpenDialogTitle", "Select a parent folder for your workspace").ToString(),
+		RootDirPreviousFolder,
+		RootDirPreviousFolder
+	);
+	
+	if (bFolderSelected)
+	{
+		RootDirTextBox->SetText(FText::FromString(FPaths::ConvertRelativePathToFull(RootDirPreviousFolder)));
+	}
+
 	return FReply::Handled();
 }
 
 FReply SNewWorkspaceWindow::OnCreateClicked()
 {
+	Tab->OnCreateWorkspace(
+		WorkspaceNameTextBox->GetText().ToString(),
+		StreamTextBox->GetText().ToString(),
+		RootDirTextBox->GetText().ToString()
+	);
+
+	Parent->SetWorkspaceTextBox(WorkspaceNameTextBox->GetText());
+
+	FSlateApplication::Get().FindWidgetWindow(AsShared())->RequestDestroyWindow();
+
 	return FReply::Handled();
 }
 
@@ -158,4 +183,17 @@ FReply SNewWorkspaceWindow::OnCancelClicked()
 	return FReply::Handled();
 }
 
+bool SNewWorkspaceWindow::IsCreateButtonEnabled() const
+{
+	const FString StreamPath = StreamTextBox->GetText().ToString();
+	const bool bStreamIsValid = StreamPath.Contains(TEXT("//")) && StreamPath.RightChop(2).Contains("/");
+
+	const bool bIsValidRootDir = FPaths::DirectoryExists(RootDirTextBox->GetText().ToString());
+
+	const FString WorkspaceFolderName = FPaths::ConvertRelativePathToFull(WorkspaceNameTextBox->GetText().ToString());
+	const bool bIsNameValid = !WorkspaceFolderName.IsEmpty();
+
+	return bStreamIsValid && bIsValidRootDir && bIsNameValid;
+}
+
 #undef LOCTEXT_NAMESPACE
diff --git a/Engine/Source/Programs/SlateUGS/Private/Widgets/SNewWorkspaceWindow.h b/Engine/Source/Programs/SlateUGS/Private/Widgets/SNewWorkspaceWindow.h
--- a/Engine/Source/Programs/SlateUGS/Private/Widgets/SNewWorkspaceWindow.h
+++ b/Engine/Source/Programs/SlateUGS/Private/Widgets/SNewWorkspaceWindow.h
@@ -2,6 +2,7 @@
 
 #pragma once
 
+#include "SWorkspaceWindow.h"
 #include "Widgets/SWindow.h"
 #include "Widgets/Input/SEditableTextBox.h"
 
@@ -13,7 +14,7 @@ public:
 	SLATE_BEGIN_ARGS(SNewWorkspaceWindow) {}
 	SLATE_END_ARGS()
 
-	void Construct(const FArguments& InArgs, UGSTab* InTab);
+	void Construct(const FArguments& InArgs, TSharedPtr<SWorkspaceWindow> InParent, UGSTab* InTab);
 
 private:
 	FReply OnBrowseStreamClicked();
@@ -22,10 +23,14 @@ private:
 	FReply OnCreateClicked();
 	FReply OnCancelClicked();
 
+	bool IsCreateButtonEnabled() const;
+
 	TSharedPtr<SEditableTextBox> StreamTextBox = nullptr;
 	TSharedPtr<SEditableTextBox> RootDirTextBox = nullptr;
-	TSharedPtr<SEditableTextBox> FileNameTextBox = nullptr;
-	FString WorkspacePathText;
+	TSharedPtr<SEditableTextBox> WorkspaceNameTextBox = nullptr;
+	FString RootDirPreviousFolder;
+
+	TSharedPtr<SWorkspaceWindow> Parent = nullptr;
 
 	UGSTab* Tab = nullptr;
 };
diff --git a/Engine/Source/Programs/SlateUGS/Private/Widgets/SWorkspaceWindow.cpp b/Engine/Source/Programs/SlateUGS/Private/Widgets/SWorkspaceWindow.cpp
--- a/Engine/Source/Programs/SlateUGS/Private/Widgets/SWorkspaceWindow.cpp
+++ b/Engine/Source/Programs/SlateUGS/Private/Widgets/SWorkspaceWindow.cpp
@@ -68,7 +68,7 @@ void SWorkspaceWindow::Construct(const FArguments& InArgs, UGSTab* InTab)
 					.Padding(10.0f, 0.0f)
 					.FillWidth(7)
 					[
-						SAssignNew(LocalFileText, SEditableTextBox)
+						SAssignNew(LocalFileTextBox, SEditableTextBox)
 						.HintText(LOCTEXT("FilePathHint", "Path/To/ProjectFile.uproject")) // Todo: Make hint text use backslash for Windows, forward slash for Unix
 			 			.OnTextChanged_Lambda([this](const FText& InText)
 						{
@@ -127,7 +127,7 @@ void SWorkspaceWindow::Construct(const FArguments& InArgs, UGSTab* InTab)
 						.Padding(10.0f, 0.0f)
 						.FillWidth(5)
 						[
-							SNew(SEditableTextBox)
+							SAssignNew(WorkspaceNameTextBox, SEditableTextBox)
 							.HintText(LOCTEXT("NameHint", "WorkspaceName"))
 						]
 						+SHorizontalBox::Slot()
@@ -209,6 +209,11 @@ void SWorkspaceWindow::Construct(const FArguments& InArgs, UGSTab* InTab)
 	]);
 }
 
+void SWorkspaceWindow::SetWorkspaceTextBox(FText Text) const
+{
+	WorkspaceNameTextBox->SetText(Text);
+}
+
 FReply SWorkspaceWindow::OnOkClicked()
 {
 	bool bIsWorkspaceValid = Tab->OnWorkspaceChosen(WorkspacePathText);
@@ -249,7 +254,7 @@ FReply SWorkspaceWindow::OnBrowseClicked()
 	if (!OutOpenFilenames.IsEmpty())
 	{
 		PreviousProjectPath = OutOpenFilenames[0];
-		LocalFileText->SetText(FText::FromString(FPaths::ConvertRelativePathToFull(PreviousProjectPath)));
+		LocalFileTextBox->SetText(FText::FromString(FPaths::ConvertRelativePathToFull(PreviousProjectPath)));
 	}
 
 	return FReply::Handled();
@@ -257,7 +262,7 @@ FReply SWorkspaceWindow::OnBrowseClicked()
 
 FReply SWorkspaceWindow::OnNewClicked()
 {
-	FSlateApplication::Get().AddModalWindow(SNew(SNewWorkspaceWindow, Tab), SharedThis(this), false);
+	FSlateApplication::Get().AddModalWindow(SNew(SNewWorkspaceWindow, SharedThis(this), Tab), SharedThis(this), false);
 
 	return FReply::Handled();
 }
diff --git a/Engine/Source/Programs/SlateUGS/Private/Widgets/SWorkspaceWindow.h b/Engine/Source/Programs/SlateUGS/Private/Widgets/SWorkspaceWindow.h
--- a/Engine/Source/Programs/SlateUGS/Private/Widgets/SWorkspaceWindow.h
+++ b/Engine/Source/Programs/SlateUGS/Private/Widgets/SWorkspaceWindow.h
@@ -15,6 +15,8 @@ public:
 
 	void Construct(const FArguments& InArgs, UGSTab* InTab);
 
+	void SetWorkspaceTextBox(FText Text) const;
+
 private:
 	FReply OnOkClicked();
 	FReply OnCancelClicked();
@@ -25,8 +27,10 @@ private:
 	FReply OnNewClicked();
 
 	bool bIsLocalFileSelected = true;
-	TSharedPtr<SEditableTextBox> LocalFileText = nullptr;
+	TSharedPtr<SEditableTextBox> LocalFileTextBox = nullptr;
 	FString WorkspacePathText;
 
+	TSharedPtr<SEditableTextBox> WorkspaceNameTextBox = nullptr;
+
 	UGSTab* Tab = nullptr;
 };
diff --git a/Engine/Source/Programs/UGSCore/CreateClientTask.cpp b/Engine/Source/Programs/UGSCore/CreateClientTask.cpp
new file mode 100644
--- /dev/null
+++ b/Engine/Source/Programs/UGSCore/CreateClientTask.cpp
@@ -0,0 +1,38 @@
+// Copyright Epic Games, Inc. All Rights Reserved.
+
+#include "CreateClientTask.h"
+
+#define LOCTEXT_NAMESPACE "UnrealGameSync"
+
+namespace UGSCore
+{
+
+	FCreateClientTask::FCreateClientTask(
+		TSharedRef<FPerforceConnection> InPerforce,
+		TSharedRef<FLineBasedTextWriter> InLog,
+		const FPerforceClientRecord& InClientRecord,
+		const FString& InStream)
+		: Perforce(MoveTemp(InPerforce))
+		, Log(MoveTemp(InLog))
+		, ClientRecord(InClientRecord)
+		, Stream(InStream)
+	{}
+
+	FCreateClientTask::~FCreateClientTask()
+	{}
+
+	TSharedRef<FModalTaskResult> FCreateClientTask::Run(FEvent* AbortEvent)
+	{
+		return RunInternal(AbortEvent);
+	}
+
+	TSharedRef<FModalTaskResult> FCreateClientTask::RunInternal(FEvent* AbortEvent) const
+	{
+		return Perforce->CreateClient(ClientRecord, Stream, AbortEvent, *Log)
+			? FModalTaskResult::Success()
+			: FModalTaskResult::Failure(LOCTEXT("CreateClientTaskFailure", "Failed to find streams"));
+	}
+
+#undef LOCTEXT_NAMESPACE
+
+} // namespace UGSCore
diff --git a/Engine/Source/Programs/UGSCore/CreateClientTask.h b/Engine/Source/Programs/UGSCore/CreateClientTask.h
new file mode 100644
--- /dev/null
+++ b/Engine/Source/Programs/UGSCore/CreateClientTask.h
@@ -0,0 +1,36 @@
+// Copyright Epic Games, Inc. All Rights Reserved.
+
+#pragma once
+
+#include "CoreMinimal.h"
+#include "ModalTask.h"
+#include "Perforce.h"
+#include "OutputAdapters.h"
+
+namespace UGSCore
+{
+
+class FCreateClientTask : public IModalTask
+{
+public:
+	TSharedRef<FPerforceConnection> Perforce;
+
+	FCreateClientTask(
+		TSharedRef<FPerforceConnection> InPerforce,
+		TSharedRef<FLineBasedTextWriter> InLog,
+		const FPerforceClientRecord& InClientRecord,
+		const FString& InStream);
+	virtual ~FCreateClientTask() override;
+
+	virtual TSharedRef<FModalTaskResult> Run(FEvent* AbortEvent) override;
+
+private:
+	TSharedRef<FLineBasedTextWriter> Log;
+
+	const FPerforceClientRecord& ClientRecord;
+	const FString& Stream;
+
+	TSharedRef<FModalTaskResult> RunInternal(FEvent* AbortEvent) const;
+};
+
+} // namespace UGSCore
diff --git a/Engine/Source/Programs/UGSCore/Perforce.cpp b/Engine/Source/Programs/UGSCore/Perforce.cpp
--- a/Engine/Source/Programs/UGSCore/Perforce.cpp
+++ b/Engine/Source/Programs/UGSCore/Perforce.cpp
@@ -86,21 +86,41 @@ struct FPerforceExe
 		return ExecutablePath;
 	}
 
-	inline int InnerRunCommand(const FString& CommandLine, TArray<FString>& OutLines, FEvent* AbortEvent)
+	inline int InnerRunCommand(const FString& CommandLine, TArray<FString>& OutLines, FEvent* AbortEvent, const FString WriteChildText = FString())
 	{
-		uint32 ProcId;
-		void* ReadPipe = nullptr;
-		void* WritePipe = nullptr;
-		FPlatformProcess::CreatePipe(ReadPipe, WritePipe);
+		void* ParentReadPipe  = nullptr;
+		void* ParentWritePipe = nullptr;
+		void* ChildReadPipe   = nullptr;
+		void* ChildWritePipe  = nullptr;
 
-		FProcHandle P4Proc = FPlatformProcess::CreateProc(*GetPerforceExe(), *CommandLine, false, true, true, &ProcId, 0, nullptr, WritePipe, ReadPipe);
+		// Create pipe for reading from child stdout
+		FPlatformProcess::CreatePipe(ParentReadPipe, ChildWritePipe);
 
+		// Create pipe for writing to child stdin
+		const bool bWriting = !WriteChildText.IsEmpty();
+		if (bWriting)
+		{
+			FPlatformProcess::CreatePipe(ChildReadPipe, ParentWritePipe);
+		}
+
+		FProcHandle P4Proc = FPlatformProcess::CreateProc(*GetPerforceExe(), *CommandLine, false, true, true, nullptr, 0, nullptr, ChildWritePipe, ChildReadPipe);
+
+		// Write to child process's stdin
+		if (FPlatformProcess::IsProcRunning(P4Proc) && bWriting)
+		{
+			FPlatformProcess::WritePipe(ParentWritePipe, WriteChildText);
+		}
+
+		// Done writing
+		FPlatformProcess::ClosePipe(ChildReadPipe, ParentWritePipe);
+
+		// Read from child process's stdout
 		FString P4Output;
-		FString LatestOutput = FPlatformProcess::ReadPipe(ReadPipe);
+		FString LatestOutput = FPlatformProcess::ReadPipe(ParentReadPipe);
 		while (FPlatformProcess::IsProcRunning(P4Proc) || !LatestOutput.IsEmpty())
 		{
 			P4Output += LatestOutput;
-			LatestOutput = FPlatformProcess::ReadPipe(ReadPipe);
+			LatestOutput = FPlatformProcess::ReadPipe(ParentReadPipe);
 
 			if (AbortEvent->Wait(FTimespan::Zero()))
 			{
@@ -109,14 +129,15 @@ struct FPerforceExe
 			}
 		}
 
-		FPlatformProcess::ClosePipe(ReadPipe, WritePipe);
+		// Done reading
+		FPlatformProcess::ClosePipe(ParentReadPipe, ChildWritePipe);
 
 		OutLines = Split(P4Output, LINE_TERMINATOR);
 
 		int ExitCode = -1;
-		bool GotReturnCode = FPlatformProcess::GetProcReturnCode(P4Proc, &ExitCode);
 
-		if (GotReturnCode)
+		const bool bGotReturnCode = FPlatformProcess::GetProcReturnCode(P4Proc, &ExitCode);
+		if (bGotReturnCode)
 		{
 			return ExitCode;
 		}
@@ -124,11 +145,11 @@ struct FPerforceExe
 		return -1;
 	}
 
-	int RunCommand(const FString& CommandLine, TArray<FString>& OutLines, FEvent* AbortEvent)
+	int RunCommand(const FString& CommandLine, TArray<FString>& OutLines, FEvent* AbortEvent, const FString& WritePipeText = FString())
 	{
 		if (bValidP4)
 		{
-			return InnerRunCommand(CommandLine, OutLines, AbortEvent);
+			return InnerRunCommand(CommandLine, OutLines, AbortEvent, WritePipeText);
 		}
 
 		return -1;
@@ -542,6 +563,15 @@ bool FPerforceConnection::TryGetClientSpec(const FString& InClientName, TSharedP
 	return true;
 }
 
+bool FPerforceConnection::CreateClient(const FPerforceClientRecord& Client, const FString& Stream, FEvent* AbortEvent, FOutputDevice& Log) const
+{
+	const FString PipeInput = FString::Printf(
+		TEXT("Client: %s\nStream: %s\nRoot: %s\nOwner: %s\nHost: %s\nDescription: Created by %s in SlateUGS\n"),
+		*Client.Name, *Stream, *Client.Root, *Client.Owner, *Client.Host, *Client.Owner);
+
+	return RunCommand(TEXT("client -i"), ECommandOptions::None, AbortEvent, Log, PipeInput);
+}
+
 bool FPerforceConnection::TryGetStreamSpec(const FString& StreamName, TSharedPtr<FPerforceSpec>& OutSpec, FEvent* AbortEvent, FOutputDevice& Log) const
 {
 	TArray<FString> Lines;
@@ -1107,25 +1137,25 @@ bool FPerforceConnection::RunCommand(const FString& CommandLine, TArray<TMap<FSt
 	return true;
 }
 
-bool FPerforceConnection::RunCommand(const FString& CommandLine, ECommandOptions Options, FEvent* AbortEvent, FOutputDevice& Log) const
+bool FPerforceConnection::RunCommand(const FString& CommandLine, ECommandOptions Options, FEvent* AbortEvent, FOutputDevice& Log, const FString& WritePipeText) const
 {
 	TArray<FString> Lines;
-	return RunCommand(CommandLine, Lines, Options, AbortEvent, Log);
+	return RunCommand(CommandLine, Lines, Options, AbortEvent, Log, WritePipeText);
 }
 
-bool FPerforceConnection::RunCommand(const FString& CommandLine, TArray<FString>& OutLines, ECommandOptions Options, FEvent* AbortEvent, FOutputDevice& Log) const
+bool FPerforceConnection::RunCommand(const FString& CommandLine, TArray<FString>& OutLines, ECommandOptions Options, FEvent* AbortEvent, FOutputDevice& Log, const FString& WritePipeText) const
 {
-	return RunCommand(CommandLine, EPerforceOutputChannel::Info, OutLines, Options, AbortEvent, Log);
+	return RunCommand(CommandLine, EPerforceOutputChannel::Info, OutLines, Options, AbortEvent, Log, WritePipeText);
 }
 
-bool FPerforceConnection::RunCommand(const FString& CommandLine, EPerforceOutputChannel Channel, TArray<FString>& OutLines, ECommandOptions Options, FEvent* AbortEvent, FOutputDevice& Log) const
+bool FPerforceConnection::RunCommand(const FString& CommandLine, EPerforceOutputChannel Channel, TArray<FString>& OutLines, ECommandOptions Options, FEvent* AbortEvent, FOutputDevice& Log, const FString& WritePipeText) const
 {
 	FString FullCommandLine = GetFullCommandLine(CommandLine, Options);
 	Log.Logf(TEXT("p4> %s %s"), *FPaths::GetCleanFilename(GPerforceExe.GetPerforceExe()), *FullCommandLine);
 
 	TArray<FString> RawOutputLines;
 
-	int ExitCode = GPerforceExe.RunCommand(FullCommandLine, RawOutputLines, AbortEvent);
+	int ExitCode = GPerforceExe.RunCommand(FullCommandLine, RawOutputLines, AbortEvent, WritePipeText);
 	if (ExitCode != 0 && !EnumHasAnyFlags(Options, ECommandOptions::IgnoreExitCode))
 	{
 		return false;
@@ -1409,11 +1439,11 @@ bool FPerforceConnection::RunCommandWithBinaryOutput(const FString& CommandLine,
 	FString FullCommandLine(TEXT("-G ") + CommandLine);
 
 	uint32 ProcId;
-	void* ReadPipe = nullptr;
-	void* WritePipe = nullptr;
-	FPlatformProcess::CreatePipe(ReadPipe, WritePipe);
+	void* ChildWritePipe  = nullptr;
+	void* ParentReadPipe = nullptr;
+	FPlatformProcess::CreatePipe(ParentReadPipe, ChildWritePipe);
 
-	FProcHandle P4Proc = FPlatformProcess::CreateProc(*GPerforceExe.GetPerforceExe(), *FullCommandLine, false, true, true, &ProcId, 0, nullptr, WritePipe, ReadPipe);
+	FProcHandle P4Proc = FPlatformProcess::CreateProc(*GPerforceExe.GetPerforceExe(), *FullCommandLine, false, true, true, &ProcId, 0, nullptr, ChildWritePipe);
 
 	size_t BufferPos = 0;
 	size_t BufferEnd = 0;
@@ -1434,7 +1464,7 @@ bool FPerforceConnection::RunCommandWithBinaryOutput(const FString& CommandLine,
 
 		// Read data from the child process
 		TArray<uint8> TempBuffer;
-		bool bRead = FPlatformProcess::ReadPipeToArray(ReadPipe, TempBuffer);
+		bool bRead = FPlatformProcess::ReadPipeToArray(ParentReadPipe, TempBuffer);
 		size_t BytesRead = TempBuffer.Num();
 
 		// Add new data to the end of the Buffer
@@ -1484,7 +1514,7 @@ bool FPerforceConnection::RunCommandWithBinaryOutput(const FString& CommandLine,
 		}
 	}
 
-	FPlatformProcess::ClosePipe(ReadPipe, WritePipe);
+	FPlatformProcess::ClosePipe(ParentReadPipe, ChildWritePipe);
 
 	int ExitCode = -1;
 
diff --git a/Engine/Source/Programs/UGSCore/Perforce.h b/Engine/Source/Programs/UGSCore/Perforce.h
--- a/Engine/Source/Programs/UGSCore/Perforce.h
+++ b/Engine/Source/Programs/UGSCore/Perforce.h
@@ -207,6 +207,7 @@ public:
 	bool GetSetting(const FString& Name, FString& Value, FEvent* AbortEvent, FOutputDevice& Log) const;
 	bool FindClients(TArray<FPerforceClientRecord>& Clients, FEvent* AbortEvent, FOutputDevice& Log) const;
 	bool FindClients(TArray<FPerforceClientRecord>& Clients, const FString& ForUserName, FEvent* AbortEvent, FOutputDevice& Log) const;
+	bool CreateClient(const FPerforceClientRecord& Client, const FString& Stream, FEvent* AbortEvent, FOutputDevice& Log) const;
 	bool TryGetClientSpec(const FString& ClientName, TSharedPtr<FPerforceSpec>& OutSpec, FEvent* AbortEvent, FOutputDevice& Log) const;
 	bool TryGetStreamSpec(const FString& StreamName, TSharedPtr<FPerforceSpec>& OutSpec, FEvent* AbortEvent, FOutputDevice& Log) const;
 	bool FindFiles(const FString& Filter, TArray<FPerforceFileRecord>& FileRecords, FEvent* AbortEvent, FOutputDevice& Log) const;
@@ -251,9 +252,9 @@ private:
 	bool RunCommand(const FString& CommandLine, TArray<FPerforceFileRecord>& OutFileRecords, ECommandOptions Options, FEvent* AbortEvent, FOutputDevice& Log) const;
 	bool RunCommand(const FString& CommandLine, TArray<FPerforceClientRecord>& OutClientRecords, ECommandOptions Options, FEvent* AbortEvent, FOutputDevice& Log) const;
 	bool RunCommand(const FString& CommandLine, TArray<TMap<FString, FString>>& OutTagRecords, ECommandOptions Options, FEvent* AbortEvent, FOutputDevice& Log) const;
-	bool RunCommand(const FString& CommandLine, ECommandOptions Options, FEvent* AbortEvent, FOutputDevice& Log) const;
-	bool RunCommand(const FString& CommandLine, TArray<FString>& OutLines, ECommandOptions Options, FEvent* AbortEvent, FOutputDevice& Log) const;
-	bool RunCommand(const FString& CommandLine, EPerforceOutputChannel Channel, TArray<FString>& OutLines, ECommandOptions Options, FEvent* AbortEvent, FOutputDevice& Log) const;
+	bool RunCommand(const FString& CommandLine, ECommandOptions Options, FEvent* AbortEvent, FOutputDevice& Log, const FString& WritePipeText = FString()) const;
+	bool RunCommand(const FString& CommandLine, TArray<FString>& OutLines, ECommandOptions Options, FEvent* AbortEvent, FOutputDevice& Log, const FString& WritePipeText = FString()) const;
+	bool RunCommand(const FString& CommandLine, EPerforceOutputChannel Channel, TArray<FString>& OutLines, ECommandOptions Options, FEvent* AbortEvent, FOutputDevice& Log, const FString& WritePipeText = FString()) const;
 	bool RunCommand(const FString& CommandLine, const TCHAR* Input, TFunction<bool(const FPerforceOutputLine&)> HandleOutput, ECommandOptions Options, FEvent* AbortEvent, FOutputDevice& Log) const;
 	FString GetFullCommandLine(const FString& CommandLine, ECommandOptions Options) const;
 	bool ParseCommandOutput(const FString& Text, TFunction<bool(const FPerforceOutputLine&)> HandleOutput, ECommandOptions Options) const;
