[misc] fix lots of FString::Printf format errors

diff --git a/Engine/Plugins/AI/HTNPlanner/Source/HTNPlanner/Private/HTNBuilder.cpp b/Engine/Plugins/AI/HTNPlanner/Source/HTNPlanner/Private/HTNBuilder.cpp
--- a/Engine/Plugins/AI/HTNPlanner/Source/HTNPlanner/Private/HTNBuilder.cpp
+++ b/Engine/Plugins/AI/HTNPlanner/Source/HTNPlanner/Private/HTNBuilder.cpp
@@ -317,7 +317,7 @@ FString FHTNBuilder_Domain::GetDebugDescription() const
 	for (auto It = PrimitiveTasks.CreateConstIterator(); It; ++It)
 	{
 		Description += FString::Printf(TEXT("%s:\n"), *It->Key.ToString());
-		Description += FString::Printf(TEXT("%\tOp: %d param: %d\n"), It->Value.ActionID, It->Value.Parameter);
+		Description += FString::Printf(TEXT("\tOp: %d param: %d\n"), It->Value.ActionID, It->Value.Parameter);
 
 		if (It->Value.Effects.Num() > 0)
 		{
diff --git a/Engine/Plugins/Animation/ControlRig/Source/ControlRig/Private/Units/Animation/RigUnit_AnimAttribute.cpp b/Engine/Plugins/Animation/ControlRig/Source/ControlRig/Private/Units/Animation/RigUnit_AnimAttribute.cpp
--- a/Engine/Plugins/Animation/ControlRig/Source/ControlRig/Private/Units/Animation/RigUnit_AnimAttribute.cpp
+++ b/Engine/Plugins/Animation/ControlRig/Source/ControlRig/Private/Units/Animation/RigUnit_AnimAttribute.cpp
@@ -288,9 +288,6 @@ FString FRigDispatch_AnimAttributeBase::GetNodeTitle(const FRigVMTemplateTypeMap
 		if(ValueTypeIndex != RigVMTypeUtils::TypeIndex::WildCard &&
 			ValueTypeIndex != RigVMTypeUtils::TypeIndex::WildCardArray)
 		{
-			static constexpr TCHAR GetFormat[] = TEXT("Get %s Animation Attribute");
-			static constexpr TCHAR SetFormat[] = TEXT("Set %s Animation Attribute");
-
 			const FRigVMTemplateArgumentType& ValueType = FRigVMRegistry::Get().GetType(ValueTypeIndex);
 			FString ValueName;
 
@@ -304,7 +301,7 @@ FString FRigDispatch_AnimAttributeBase::GetNodeTitle(const FRigVMTemplateTypeMap
 				ValueName = ValueName.Left(1).ToUpper() + ValueName.Mid(1);
 			}
 
-			return FString::Printf(IsSet() ? SetFormat : GetFormat, *ValueName); 
+			return FString::Printf(TEXT("%s %s Animation Attribute"), IsSet() ? TEXT("Set") : TEXT("Get"), *ValueName);
 		}
 	}
 	return FRigDispatchFactory::GetNodeTitle(InTypes);	
diff --git a/Engine/Plugins/Animation/ControlRig/Source/ControlRig/Private/Units/Hierarchy/RigUnit_Metadata.cpp b/Engine/Plugins/Animation/ControlRig/Source/ControlRig/Private/Units/Hierarchy/RigUnit_Metadata.cpp
--- a/Engine/Plugins/Animation/ControlRig/Source/ControlRig/Private/Units/Hierarchy/RigUnit_Metadata.cpp
+++ b/Engine/Plugins/Animation/ControlRig/Source/ControlRig/Private/Units/Hierarchy/RigUnit_Metadata.cpp
@@ -23,9 +23,6 @@ FString FRigDispatch_MetadataBase::GetNodeTitle(const FRigVMTemplateTypeMap& InT
 		if(ValueTypeIndex != RigVMTypeUtils::TypeIndex::WildCard &&
 			ValueTypeIndex != RigVMTypeUtils::TypeIndex::WildCardArray)
 		{
-			static constexpr TCHAR GetMetadataFormat[] = TEXT("Get %s%s Metadata");
-			static constexpr TCHAR SetMetadataFormat[] = TEXT("Set %s%s Metadata");
-
 			const FRigVMTemplateArgumentType& ValueType = FRigVMRegistry::Get().GetType(ValueTypeIndex);
 			FString ValueName;
 			if(ValueType.CPPTypeObject == FRigElementKey::StaticStruct())
@@ -51,7 +48,7 @@ FString FRigDispatch_MetadataBase::GetNodeTitle(const FRigVMTemplateTypeMap& InT
 				ValueName += TEXT(" Array");
 			}
 
-			return FString::Printf(IsSetMetadata() ? SetMetadataFormat : GetMetadataFormat, *GetNodeTitlePrefix(), *ValueName); 
+			return FString::Printf(TEXT("%s %s%s Metadata"), IsSetMetadata() ? TEXT("Set") : TEXT("Get"), *GetNodeTitlePrefix(), *ValueName);
 		}
 	}
 	return FRigDispatchFactory::GetNodeTitle(InTypes);
diff --git a/Engine/Plugins/Animation/PoseSearch/Source/Runtime/Private/AnimNode_MotionMatchingInteraction.cpp b/Engine/Plugins/Animation/PoseSearch/Source/Runtime/Private/AnimNode_MotionMatchingInteraction.cpp
--- a/Engine/Plugins/Animation/PoseSearch/Source/Runtime/Private/AnimNode_MotionMatchingInteraction.cpp
+++ b/Engine/Plugins/Animation/PoseSearch/Source/Runtime/Private/AnimNode_MotionMatchingInteraction.cpp
@@ -23,7 +23,7 @@ void FAnimNode_MotionMatchingInteraction::GatherDebugData(FNodeDebugData& DebugD
 	FString DebugLine = DebugData.GetNodeName(this);
 #if ENABLE_ANIM_DEBUG
 	static bool bActive = true;
-	DebugLine += FString::Printf(TEXT("\n - Active: (%s)"), bActive);
+	DebugLine += FString::Printf(TEXT("\n - Active: (%s)"), bActive ? TEXT("True") : TEXT("False"));
 #endif
 	DebugData.AddDebugItem(DebugLine);
 }
diff --git a/Engine/Plugins/Enterprise/DatasmithC4DImporter/Source/DatasmithC4DTranslator/Private/DatasmithC4DUtils.cpp b/Engine/Plugins/Enterprise/DatasmithC4DImporter/Source/DatasmithC4DTranslator/Private/DatasmithC4DUtils.cpp
--- a/Engine/Plugins/Enterprise/DatasmithC4DImporter/Source/DatasmithC4DTranslator/Private/DatasmithC4DUtils.cpp
+++ b/Engine/Plugins/Enterprise/DatasmithC4DImporter/Source/DatasmithC4DTranslator/Private/DatasmithC4DUtils.cpp
@@ -334,7 +334,7 @@ FString GeDataToString(const cineware::GeData& Data)
 		return MelangeMatrixToFMatrix(Data.GetMatrix()).ToString();
 		break;
 	case cineware::DA_LLONG:
-		return FString::Printf(TEXT("%jd"), Data.GetInt64());
+		return FString::Printf(TEXT("%" INT64_FMT), Data.GetInt64());
 		break;
 	case cineware::DA_BYTEARRAY:
 	{
@@ -438,4 +438,4 @@ TOptional<FString> GetMelangeBaseList2dID(cineware::BaseList2D* BaseList)
 	return MelangeID;
 }
 
-#endif
\ No newline at end of file
+#endif
diff --git a/Engine/Plugins/Enterprise/DatasmithFBXImporter/Source/DatasmithFBXTranslator/Private/DatasmithFBXFileImporter.cpp b/Engine/Plugins/Enterprise/DatasmithFBXImporter/Source/DatasmithFBXTranslator/Private/DatasmithFBXFileImporter.cpp
--- a/Engine/Plugins/Enterprise/DatasmithFBXImporter/Source/DatasmithFBXTranslator/Private/DatasmithFBXFileImporter.cpp
+++ b/Engine/Plugins/Enterprise/DatasmithFBXImporter/Source/DatasmithFBXTranslator/Private/DatasmithFBXFileImporter.cpp
@@ -168,7 +168,7 @@ namespace FBXFileImporterImpl
 		case eFbxDouble2:
 		{
 			FbxDouble2 Vec = Property.Get<FbxDouble2>();
-			ValueStr = FString::Printf( TEXT( "(%f, %f, %f, %f)" ), Vec[ 0 ], Vec[ 1 ] );
+			ValueStr = FString::Printf( TEXT( "(%f, %f)" ), Vec[ 0 ], Vec[ 1 ] );
 		}
 		break;
 		case eFbxDouble3:
diff --git a/Engine/Plugins/Experimental/AVCodecs/LibVpxCodecs/Source/LibVpxCodecs/Public/Video/Encoders/VideoEncoderLibVpxVP8.hpp b/Engine/Plugins/Experimental/AVCodecs/LibVpxCodecs/Source/LibVpxCodecs/Public/Video/Encoders/VideoEncoderLibVpxVP8.hpp
--- a/Engine/Plugins/Experimental/AVCodecs/LibVpxCodecs/Source/LibVpxCodecs/Public/Video/Encoders/VideoEncoderLibVpxVP8.hpp
+++ b/Engine/Plugins/Experimental/AVCodecs/LibVpxCodecs/Source/LibVpxCodecs/Public/Video/Encoders/VideoEncoderLibVpxVP8.hpp
@@ -13,7 +13,7 @@ using namespace UE::AVCodecCore::VP8;
 		if (Res != VPX_CODEC_OK)                                                                                                                                                               \
 		{                                                                                                                                                                                      \
 			FString ErrorString(::vpx_codec_error_detail(Encoder));                                                                                                                            \
-			FAVResult::Log(EAVResult::Success, FString::Printf(TEXT("Error executing ::vpx_codec_control. Setting %s, Error: %d, Details: %s"), #Setting, Res, *ErrorString), TEXT("LibVpx")); \
+			FAVResult::Log(EAVResult::Success, FString::Printf(TEXT("Error executing ::vpx_codec_control. Setting %s, Error: %d, Details: %s"), TEXT(#Setting), Res, *ErrorString), TEXT("LibVpx")); \
 		}                                                                                                                                                                                      \
 	}
 
@@ -513,4 +513,4 @@ typename TVideoEncoderLibVpxVP8<TResource>::EEncodeResult TVideoEncoderLibVpxVP8
 	}
 
 	return EEncodeResult::Success;
-}
\ No newline at end of file
+}
diff --git a/Engine/Plugins/Experimental/AVCodecs/LibVpxCodecs/Source/LibVpxCodecs/Public/Video/Encoders/VideoEncoderLibVpxVP9.hpp b/Engine/Plugins/Experimental/AVCodecs/LibVpxCodecs/Source/LibVpxCodecs/Public/Video/Encoders/VideoEncoderLibVpxVP9.hpp
--- a/Engine/Plugins/Experimental/AVCodecs/LibVpxCodecs/Source/LibVpxCodecs/Public/Video/Encoders/VideoEncoderLibVpxVP9.hpp
+++ b/Engine/Plugins/Experimental/AVCodecs/LibVpxCodecs/Source/LibVpxCodecs/Public/Video/Encoders/VideoEncoderLibVpxVP9.hpp
@@ -766,7 +766,7 @@ FAVResult TVideoEncoderLibVpxVP9<TResource>::InitAndSetControlSettings(FVideoEnc
 	const vpx_codec_err_t Result = ::vpx_codec_enc_init(Encoder.Get(), vpx_codec_vp9_cx(), VpxConfig.Get(), VpxConfig->g_bit_depth == VPX_BITS_8 ? 0 : VPX_CODEC_USE_HIGHBITDEPTH);
 	if (Result != VPX_CODEC_OK)
 	{
-		return FAVResult(EAVResult::Error, FString::Printf(TEXT("Init error %s"), ::vpx_codec_err_to_string(Result)), TEXT("LibVpx"));
+		return FAVResult(EAVResult::Error, FString::Printf(TEXT("Init error %hs"), ::vpx_codec_err_to_string(Result)), TEXT("LibVpx"));
 	}
 
 	if (PerformanceFlags.bUsePerLayerSpeed)
@@ -933,18 +933,18 @@ TUniquePtr<FScalableVideoController> TVideoEncoderLibVpxVP9<TResource>::CreateSc
 	TOptional<EScalabilityMode> ScalabilityMode = ScalabilityModeFromString(Builder.ToString());
 	if (!ScalabilityMode.IsSet())
 	{
-		FAVResult::Log(EAVResult::Warning, FString::Printf(TEXT("Invalid scalability mode %s"), *Builder.ToString()), TEXT("LibVpx"));
+		FAVResult::Log(EAVResult::Warning, FString::Printf(TEXT("Invalid scalability mode %s"), Builder.ToString()), TEXT("LibVpx"));
 		return nullptr;
 	}
 
 	TUniquePtr<FScalableVideoController> ScalabilityStructureController = CreateScalabilityStructure(*ScalabilityMode);
 	if (!ScalabilityStructureController)
 	{
-		FAVResult::Log(EAVResult::Warning, FString::Printf(TEXT("Unsupported scalability structure %s"), *Builder.ToString()), TEXT("LibVpx"));
+		FAVResult::Log(EAVResult::Warning, FString::Printf(TEXT("Unsupported scalability structure %s"), Builder.ToString()), TEXT("LibVpx"));
 	}
 	else
 	{
-		FAVResult::Log(EAVResult::Success, FString::Printf(TEXT("Created scalability structure %s"), *Builder.ToString()), TEXT("LibVpx"));
+		FAVResult::Log(EAVResult::Success, FString::Printf(TEXT("Created scalability structure %s"), Builder.ToString()), TEXT("LibVpx"));
 	}
 	return ScalabilityStructureController;
 }
@@ -1343,7 +1343,7 @@ void TVideoEncoderLibVpxVP9<TResource>::FillReferenceIndices(const vpx_codec_cx_
 			}
 		}
 
-		FAVResult::Log(EAVResult::Success, FString::Printf(TEXT("Frame %d S%dT%d referenced buffers %s"), PicNum, LayerId.spatial_layer_id, LayerId.temporal_layer_id, RefBufFlags), TEXT("LibVpx"));
+		FAVResult::Log(EAVResult::Success, FString::Printf(TEXT("Frame %d S%dT%d referenced buffers %hs"), PicNum, LayerId.spatial_layer_id, LayerId.temporal_layer_id, RefBufFlags), TEXT("LibVpx"));
 	}
 	else if (!bIsKeyFrame)
 	{
@@ -1808,4 +1808,4 @@ TTuple<size_t, size_t> TVideoEncoderLibVpxVP9<TResource>::GetActiveLayers(const
 		}
 	}
 	return { 0, 0 };
-}
\ No newline at end of file
+}
diff --git a/Engine/Plugins/Experimental/TextureGraph/Source/TextureGraphEngine/Device/DeviceNativeTask.cpp b/Engine/Plugins/Experimental/TextureGraph/Source/TextureGraphEngine/Device/DeviceNativeTask.cpp
--- a/Engine/Plugins/Experimental/TextureGraph/Source/TextureGraphEngine/Device/DeviceNativeTask.cpp
+++ b/Engine/Plugins/Experimental/TextureGraph/Source/TextureGraphEngine/Device/DeviceNativeTask.cpp
@@ -97,7 +97,7 @@ ENamedThreads::Type DeviceNativeTask::GetExecutionThread() const
 
 FString DeviceNativeTask::GetDebugName() const
 {
-	return FString::Printf(TEXT("%s [%ju, %d]"), *GetName(), GetTaskId(), GetPriority());
+	return FString::Printf(TEXT("%s [%" INT64_FMT ", %d]"), *GetName(), GetTaskId(), GetPriority());
 }
 
 FString DeviceNativeTask::GetName() const
diff --git a/Engine/Plugins/Experimental/TextureGraph/Source/TextureGraphEngine/Job/Job.cpp b/Engine/Plugins/Experimental/TextureGraph/Source/TextureGraphEngine/Job/Job.cpp
--- a/Engine/Plugins/Experimental/TextureGraph/Source/TextureGraphEngine/Job/Job.cpp
+++ b/Engine/Plugins/Experimental/TextureGraph/Source/TextureGraphEngine/Job/Job.cpp
@@ -1146,7 +1146,7 @@ AsyncInt Job::PreExecAsync(ENamedThreads::Type execThread, ENamedThreads::Type R
 
 FString Job::GetDebugName() const
 {
-	return FString::Printf(TEXT("%s [%ju, %d, ID: %ju]"), *GetName(), GetTaskId(), GetPriority(), Id);
+	return FString::Printf(TEXT("%s [%" INT64_FMT ", %d, ID: %" INT64_FMT "]"), *GetName(), GetTaskId(), GetPriority(), Id);
 }
 
 FString Job::GetName() const
diff --git a/Engine/Plugins/FX/Niagara/Source/Niagara/Private/NiagaraGpuComputeDispatch.cpp b/Engine/Plugins/FX/Niagara/Source/Niagara/Private/NiagaraGpuComputeDispatch.cpp
--- a/Engine/Plugins/FX/Niagara/Source/Niagara/Private/NiagaraGpuComputeDispatch.cpp
+++ b/Engine/Plugins/FX/Niagara/Source/Niagara/Private/NiagaraGpuComputeDispatch.cpp
@@ -794,7 +794,7 @@ void FNiagaraGpuComputeDispatch::PrepareTicksForProxy(FRHICommandListImmediate&
 					if (bRaisedWarningThisFrame == false)
 					{
 						bRaisedWarningThisFrame = true;
-						GEngine->AddOnScreenDebugMessage(uint64(this), 1.f, FColor::White, *FString::Printf(TEXT("GPU Simulation(%s) will not show in preview mode, as we may sample from wrong SceneTextures buffer."), *ComputeContext->GetDebugSimName()));
+						GEngine->AddOnScreenDebugMessage(uint64(this), 1.f, FColor::White, *FString::Printf(TEXT("GPU Simulation(%s) will not show in preview mode, as we may sample from wrong SceneTextures buffer."), ComputeContext->GetDebugSimName()));
 					}
 					continue;
 				}
diff --git a/Engine/Plugins/FastBuildController/Source/Private/FastBuildUtilities.cpp b/Engine/Plugins/FastBuildController/Source/Private/FastBuildUtilities.cpp
--- a/Engine/Plugins/FastBuildController/Source/Private/FastBuildUtilities.cpp
+++ b/Engine/Plugins/FastBuildController/Source/Private/FastBuildUtilities.cpp
@@ -231,7 +231,7 @@ void FastBuildUtilities::WriteDependenciesForShaderToScript(const TArray<FTask*>
 
 void FastBuildUtilities::FASTBuildWriteScriptFileHeader(const TArray<FTask*>& InCompilationTasks, FArchive& ScriptFile, const FString& WorkerName)
 {	
-	static const TCHAR HeaderTemplate[] =
+	static constexpr TCHAR HeaderTemplate[] =
 		TEXT("Settings" LINE_TERMINATOR_ANSI)
 		TEXT("{" LINE_TERMINATOR_ANSI)
 		TEXT("\t.CachePath = '%s'" LINE_TERMINATOR_ANSI)
diff --git a/Engine/Plugins/Media/ElectraCodecs/Source/ElectraDecoders/Private/Utils/MPEG/ElectraUtilsMPEGVideo.cpp b/Engine/Plugins/Media/ElectraCodecs/Source/ElectraDecoders/Private/Utils/MPEG/ElectraUtilsMPEGVideo.cpp
--- a/Engine/Plugins/Media/ElectraCodecs/Source/ElectraDecoders/Private/Utils/MPEG/ElectraUtilsMPEGVideo.cpp
+++ b/Engine/Plugins/Media/ElectraCodecs/Source/ElectraDecoders/Private/Utils/MPEG/ElectraUtilsMPEGVideo.cpp
@@ -1248,7 +1248,7 @@ namespace ElectraDecodersUtil
 			}
 			else
 			{
-				return FString::Printf(TEXT("%s.%c%d.%X.%c%d%s"), SampleTypePrefix, TCHAR('A')+general_profile_space-1, general_profile_idc, cf, general_tier_flag ? TCHAR('H') : TCHAR('L'), general_level_idc, *cfs);
+				return FString::Printf(TEXT("%s.%c%d.%X.%c%d%s"), SampleTypePrefix, TCHAR(TCHAR('A')+general_profile_space-1), general_profile_idc, cf, general_tier_flag ? TCHAR('H') : TCHAR('L'), general_level_idc, *cfs);
 			}
 		}
 
diff --git a/Engine/Plugins/Online/OnlineServices/Source/OnlineServicesInterface/Private/Online/Achievements.cpp b/Engine/Plugins/Online/OnlineServices/Source/OnlineServicesInterface/Private/Online/Achievements.cpp
--- a/Engine/Plugins/Online/OnlineServices/Source/OnlineServicesInterface/Private/Online/Achievements.cpp
+++ b/Engine/Plugins/Online/OnlineServices/Source/OnlineServicesInterface/Private/Online/Achievements.cpp
@@ -12,7 +12,7 @@ FString ToLogString(const FAchievementStatDefinition& Definition)
 
 FString ToLogString(const FAchievementDefinition& Definition)
 {
-	const TCHAR FormatString[] = TEXT(
+	static constexpr TCHAR FormatString[] = TEXT(
 		"AchievementId=\"%s\" "
 		"UnlockedDisplayName=\"%s\" "
 		"UnlockedDescription=\"%s\" "
diff --git a/Engine/Plugins/Online/OnlineSubsystemEOS/Source/OnlineSubsystemEOS/Private/OnlineAchievementsEOS.cpp b/Engine/Plugins/Online/OnlineSubsystemEOS/Source/OnlineSubsystemEOS/Private/OnlineAchievementsEOS.cpp
--- a/Engine/Plugins/Online/OnlineSubsystemEOS/Source/OnlineSubsystemEOS/Private/OnlineAchievementsEOS.cpp
+++ b/Engine/Plugins/Online/OnlineSubsystemEOS/Source/OnlineSubsystemEOS/Private/OnlineAchievementsEOS.cpp
@@ -29,7 +29,7 @@ void FOnlineAchievementsEOS::UnlockAchievements(const FUniqueNetId& PlayerId, FO
 	const int32 LocalUserId = EOSSubsystem->UserManager->GetLocalUserNumFromUniqueNetId(PlayerId);
 	if (LocalUserId < 0)
 	{
-		UE_LOG_ONLINE_ACHIEVEMENTS(Error, TEXT("Can't unlock achievements for non-local user (%)"), *PlayerId.ToString());
+		UE_LOG_ONLINE_ACHIEVEMENTS(Error, TEXT("Can't unlock achievements for non-local user (%s)"), *PlayerId.ToString());
 		Delegate.ExecuteIfBound(PlayerId, false);
 		return;
 	}
@@ -100,7 +100,7 @@ void FOnlineAchievementsEOS::QueryAchievements(const FUniqueNetId& PlayerId, con
 	int32 LocalUserId = EOSSubsystem->UserManager->GetLocalUserNumFromUniqueNetId(PlayerId);
 	if (LocalUserId < 0)
 	{
-		UE_LOG_ONLINE_ACHIEVEMENTS(Error, TEXT("Can't query achievement progress for non-local user (%)"), *PlayerId.ToString());
+		UE_LOG_ONLINE_ACHIEVEMENTS(Error, TEXT("Can't query achievement progress for non-local user (%s)"), *PlayerId.ToString());
 		Delegate.ExecuteIfBound(PlayerId, false);
 		return;
 	}
@@ -184,7 +184,7 @@ void FOnlineAchievementsEOS::QueryAchievementDescriptions(const FUniqueNetId& Pl
 	int32 LocalUserId = EOSSubsystem->UserManager->GetLocalUserNumFromUniqueNetId(PlayerId);
 	if (LocalUserId < 0)
 	{
-		UE_LOG_ONLINE_ACHIEVEMENTS(Error, TEXT("Can't query achievement definitions for non-local user (%)"), *PlayerId.ToString());
+		UE_LOG_ONLINE_ACHIEVEMENTS(Error, TEXT("Can't query achievement definitions for non-local user (%s)"), *PlayerId.ToString());
 		Delegate.ExecuteIfBound(PlayerId, false);
 		return;
 	}
diff --git a/Engine/Plugins/Online/OnlineSubsystemEOS/Source/OnlineSubsystemEOS/Private/OnlineUserCloudEOS.cpp b/Engine/Plugins/Online/OnlineSubsystemEOS/Source/OnlineSubsystemEOS/Private/OnlineUserCloudEOS.cpp
--- a/Engine/Plugins/Online/OnlineSubsystemEOS/Source/OnlineSubsystemEOS/Private/OnlineUserCloudEOS.cpp
+++ b/Engine/Plugins/Online/OnlineSubsystemEOS/Source/OnlineSubsystemEOS/Private/OnlineUserCloudEOS.cpp
@@ -244,7 +244,7 @@ void FOnlineUserCloudEOS::EnumerateUserFiles(const FUniqueNetId& UserId)
 		}
 		else
 		{
-			UE_LOG_ONLINE_CLOUD(Warning, TEXT("[FOnlineUserCloudEOS::EnumerateUserFiles] UniqueNetId not found for EOS_ProductUserId %s"), Data->LocalUserId);
+			UE_LOG_ONLINE_CLOUD(Warning, TEXT("[FOnlineUserCloudEOS::EnumerateUserFiles] UniqueNetId not found for EOS_ProductUserId %d"), Data->LocalUserId);
 		}
 
 		TriggerOnEnumerateUserFilesCompleteDelegates(bWasSuccessful, *UniqueNetIdPtr);
diff --git a/Engine/Plugins/Runtime/GameInput/Source/GameInputBase/Private/GameInputUtils.cpp b/Engine/Plugins/Runtime/GameInput/Source/GameInputBase/Private/GameInputUtils.cpp
--- a/Engine/Plugins/Runtime/GameInput/Source/GameInputBase/Private/GameInputUtils.cpp
+++ b/Engine/Plugins/Runtime/GameInput/Source/GameInputBase/Private/GameInputUtils.cpp
@@ -15,9 +15,10 @@ namespace UE::GameInput
 			const GameInputDeviceInfo* Info = Device->GetDeviceInfo();
 			const bool bIsVirtualDevice = Info->deviceFamily == GameInputFamilyVirtual;
 			
-			Result = FString::Printf(TEXT("%llp   DeviceName: %s %s ProdId: %04x (%u) VendId: %04x (%u)"), 
+			Result = FString::Printf(TEXT("%p   DeviceName: %.*hs %s ProdId: %04x (%u) VendId: %04x (%u)"),
 				Device, 
-				Info->displayName, 
+				Info->displayName->sizeInBytes,
+				Info->displayName->data,
 				bIsVirtualDevice ? TEXT("(virtual)") : TEXT(""),
 				Info->productId, (uint32)Info->productId,
 				Info->vendorId, (uint32)Info->vendorId);
@@ -175,4 +176,4 @@ namespace UE::GameInput
 	}
 }
 
-#endif	// #if GAME_INPUT_SUPPORT
\ No newline at end of file
+#endif	// #if GAME_INPUT_SUPPORT
diff --git a/Engine/Plugins/Runtime/Harmonix/Source/HarmonixDsp/Private/HarmonixDsp/AudioBuffer.cpp b/Engine/Plugins/Runtime/Harmonix/Source/HarmonixDsp/Private/HarmonixDsp/AudioBuffer.cpp
--- a/Engine/Plugins/Runtime/Harmonix/Source/HarmonixDsp/Private/HarmonixDsp/AudioBuffer.cpp
+++ b/Engine/Plugins/Runtime/Harmonix/Source/HarmonixDsp/Private/HarmonixDsp/AudioBuffer.cpp
@@ -81,7 +81,7 @@ void HarmonixDsp::FAudioBuffer::DebugLog(const int16* InData, uint64 InNumSample
 			Db = -99.9f;
 		}
 		constexpr int32 ColumnWidth = 6;
-		DebugString += FString::Printf(TEXT("%u: %d %s (%f db)\n"), SampleIdx, ColumnWidth, InData[SampleIdx], *GetAsciiArtLevel(Sample), Db);
+		DebugString += FString::Printf(TEXT("%u: %d %d %s (%f db)\n"), SampleIdx, ColumnWidth, InData[SampleIdx], *GetAsciiArtLevel(Sample), Db);
 	}
 	DebugString += FString("------------------------------------\n");
 
@@ -135,4 +135,4 @@ void HarmonixDsp::FAudioBuffer::Convert(const TAudioBuffer<int16>& Source, TAudi
 			DstData[SampleIdx] = Sample;
 		}
 	}
-}
\ No newline at end of file
+}
diff --git a/Engine/Plugins/Runtime/RigVM/Source/RigVM/Private/RigVMFunctions/RigVMDispatch_Array.cpp b/Engine/Plugins/Runtime/RigVM/Source/RigVM/Private/RigVMFunctions/RigVMDispatch_Array.cpp
--- a/Engine/Plugins/Runtime/RigVM/Source/RigVM/Private/RigVMFunctions/RigVMDispatch_Array.cpp
+++ b/Engine/Plugins/Runtime/RigVM/Source/RigVM/Private/RigVMFunctions/RigVMDispatch_Array.cpp
@@ -1020,7 +1020,7 @@ void FRigVMDispatch_ArrayFind::Execute(FRigVMExtendedExecuteContext& InContext,
 	}
 	else
 	{
-		static constexpr TCHAR IncompatibleTypes[] = TEXT("Array('%s') doesn't support searching for element('%$s').");
+		static constexpr TCHAR IncompatibleTypes[] = TEXT("Array('%s') doesn't support searching for element('%s').");
 		InContext.GetPublicData<>().Logf(EMessageSeverity::Error, IncompatibleTypes, *PropertyB->GetCPPType(), *PropertyA->GetCPPType());
 	}
 }
diff --git a/Engine/Plugins/Runtime/RigVM/Source/RigVM/Public/RigVMCore/RigVMMemoryStorage.h b/Engine/Plugins/Runtime/RigVM/Source/RigVM/Public/RigVMCore/RigVMMemoryStorage.h
--- a/Engine/Plugins/Runtime/RigVM/Source/RigVM/Public/RigVMCore/RigVMMemoryStorage.h
+++ b/Engine/Plugins/Runtime/RigVM/Source/RigVM/Public/RigVMCore/RigVMMemoryStorage.h
@@ -103,8 +103,7 @@ struct RIGVM_API FRigVMBranchInfo
 
 	static FString GetFixedArrayLabel(const FString& InArrayName, const FString& InElementName)
 	{
-		static constexpr TCHAR Format[] = TEXT("%s_%s");
-		return FString::Printf(Format, *InArrayName, *InElementName);
+		return FString::Printf(TEXT("%s_%s"), *InArrayName, *InElementName);
 	}
 
 	static FName GetFixedArrayLabel(const FName& InArrayName, const FName& InElementName)
diff --git a/Engine/Plugins/Runtime/RigVM/Source/RigVMDeveloper/Private/RigVMModel/RigVMClient.cpp b/Engine/Plugins/Runtime/RigVM/Source/RigVMDeveloper/Private/RigVMModel/RigVMClient.cpp
--- a/Engine/Plugins/Runtime/RigVM/Source/RigVMDeveloper/Private/RigVMModel/RigVMClient.cpp
+++ b/Engine/Plugins/Runtime/RigVM/Source/RigVMDeveloper/Private/RigVMModel/RigVMClient.cpp
@@ -346,7 +346,7 @@ void FRigVMClient::RefreshAllModels(ERigVMLoadType InLoadType, bool bEnablePostL
 
 						if (UnitNode->GetScriptStruct() == nullptr)
 						{
-							static const TCHAR UnresolvedUnitNodeMessage[] = TEXT("Node %s could not be resolved.");
+							static constexpr TCHAR UnresolvedUnitNodeMessage[] = TEXT("Node %s could not be resolved.");
 							Controller->ReportErrorf(UnresolvedUnitNodeMessage, *ModelNode->GetNodePath(true));
 						}
 					}
@@ -355,7 +355,7 @@ void FRigVMClient::RefreshAllModels(ERigVMLoadType InLoadType, bool bEnablePostL
 				{
 					if (DispatchNode->GetFactory() == nullptr)
 					{
-						static const TCHAR UnresolvedDispatchNodeMessage[] = TEXT("Dispatch node %s has no factory..");
+						static constexpr TCHAR UnresolvedDispatchNodeMessage[] = TEXT("Dispatch node %s has no factory..");
 						Controller->ReportErrorf(UnresolvedDispatchNodeMessage, *ModelNode->GetNodePath(true));
 					}
 					else if (!DispatchNode->HasWildCardPin())
@@ -366,7 +366,7 @@ void FRigVMClient::RefreshAllModels(ERigVMLoadType InLoadType, bool bEnablePostL
 						}
 						if (DispatchNode->GetResolvedFunction() == nullptr)
 						{
-							static const TCHAR UnresolvedDispatchNodeMessage[] = TEXT("Node %s could not be resolved.");
+							static constexpr TCHAR UnresolvedDispatchNodeMessage[] = TEXT("Node %s could not be resolved.");
 							Controller->ReportErrorf(UnresolvedDispatchNodeMessage, *ModelNode->GetNodePath(true));
 						}
 					}
diff --git a/Engine/Plugins/Tests/RHITests/Source/RHITests/Public/RHITestsCommon.h b/Engine/Plugins/Tests/RHITests/Source/RHITests/Public/RHITestsCommon.h
--- a/Engine/Plugins/Tests/RHITests/Source/RHITests/Public/RHITestsCommon.h
+++ b/Engine/Plugins/Tests/RHITests/Source/RHITests/Public/RHITestsCommon.h
@@ -18,7 +18,7 @@ bool RunOnRenderThreadSynchronous(TFunctionRef<bool(FRHICommandListImmediate&)>
 template <typename ValueType>
 static inline FString ClearValueToString(const ValueType & ClearValue)
 {
-	if constexpr (std::is_same_v<ValueType, FVector4>)
+	if constexpr (std::is_floating_point_v<decltype(ClearValue.X)>)
 	{
 		return FString::Printf(TEXT("%f %f %f %f"), ClearValue.X, ClearValue.Y, ClearValue.Z, ClearValue.W);
 	}
diff --git a/Engine/Plugins/Tests/WidgetAutomationTests/Source/WidgetAutomationTests/Private/STextBlockTest.cpp b/Engine/Plugins/Tests/WidgetAutomationTests/Source/WidgetAutomationTests/Private/STextBlockTest.cpp
--- a/Engine/Plugins/Tests/WidgetAutomationTests/Source/WidgetAutomationTests/Private/STextBlockTest.cpp
+++ b/Engine/Plugins/Tests/WidgetAutomationTests/Source/WidgetAutomationTests/Private/STextBlockTest.cpp
@@ -88,7 +88,7 @@ public:
 	{
 		if (GetSubjectWidget()->GetTestJSONObjectToWrite().IsValid() && !FSlateTestHelper::WriteJSONToTxt(GetSubjectWidget()->GetTestJSONObjectToWrite().ToSharedRef(), GetTestName()))
 		{
-			TestObj->AddError(FString::Printf(TEXT("Failed to write the ground truth JSON file for the test %s."), "ColorAndOpacity", *GetTestName()));
+			TestObj->AddError(FString::Printf(TEXT("Failed to write the ground truth JSON file for the test %s."), *GetTestName()));
 		}
 		FSlateApplication::Get().RequestDestroyWindow(GetMainWindow());
 	}
@@ -119,4 +119,4 @@ bool FTextBlockTest::RunTest(const FString& Parameters)
 	return true;
 }
 }
-#endif
\ No newline at end of file
+#endif
diff --git a/Engine/Plugins/VirtualProduction/RemoteControl/Source/RemoteControl/Private/RemoteControlModule.cpp b/Engine/Plugins/VirtualProduction/RemoteControl/Source/RemoteControl/Private/RemoteControlModule.cpp
--- a/Engine/Plugins/VirtualProduction/RemoteControl/Source/RemoteControl/Private/RemoteControlModule.cpp
+++ b/Engine/Plugins/VirtualProduction/RemoteControl/Source/RemoteControl/Private/RemoteControlModule.cpp
@@ -1356,7 +1356,7 @@ bool FRemoteControlModule::ResolveObjectProperty(ERCAccess AccessType, UObject*
 	}
 	else
 	{
-		ErrorText = FString::Printf(TEXT("Can't resolve object '%s' properties '%s' : %s while saving or garbage collecting."), *Object->GetPathName(), *PropertyPath.GetFieldName().ToString());
+		ErrorText = FString::Printf(TEXT("Can't resolve object '%s' properties: %s while saving or garbage collecting."), *Object->GetPathName(), *PropertyPath.GetFieldName().ToString());
 		bSuccess = false;
 	}
 
diff --git a/Engine/Source/Programs/UGSCore/PerforceMonitor.cpp b/Engine/Source/Programs/UGSCore/PerforceMonitor.cpp
--- a/Engine/Source/Programs/UGSCore/PerforceMonitor.cpp
+++ b/Engine/Source/Programs/UGSCore/PerforceMonitor.cpp
@@ -195,7 +195,7 @@ void FPerforceMonitor::RunInternal()
 		}
 		else
 		{
-			LastStatusMessage = FString::Printf(TEXT("Last update took %dms"), (FDateTime::UtcNow() - StartTime).GetTotalMilliseconds());
+			LastStatusMessage = FString::Printf(TEXT("Last update took %dms"), int((FDateTime::UtcNow() - StartTime).GetTotalMilliseconds()));
 		}
 
 		// Wait for another request, or scan for new builds after a timeout
diff --git a/Engine/Source/Programs/UGSCore/Workspace.cpp b/Engine/Source/Programs/UGSCore/Workspace.cpp
--- a/Engine/Source/Programs/UGSCore/Workspace.cpp
+++ b/Engine/Source/Programs/UGSCore/Workspace.cpp
@@ -577,7 +577,7 @@ EWorkspaceUpdateResult FWorkspace::UpdateWorkspaceInternal(FWorkspaceUpdateConte
 					BuildVersionStrings.Add(TEXT("\"IsPromotedBuild\":"), TEXT(" 0,"));
 					if(!UpdateVersionFile(*(ClientRootPath + BuildVersionFileName), BuildVersionStrings, PendingChangeNumber))
 					{
-						OutStatusMessage = FString::Printf(TEXT("Failed to update %s."), *BuildVersionFileName);
+						OutStatusMessage = FString::Printf(TEXT("Failed to update %s."), BuildVersionFileName);
 						return EWorkspaceUpdateResult::FailedToSync;
 					}
 
@@ -587,12 +587,12 @@ EWorkspaceUpdateResult FWorkspace::UpdateWorkspaceInternal(FWorkspaceUpdateConte
 					VersionHeaderStrings.Add(TEXT("#define BRANCH_NAME"), TEXT(" \"") + BranchOrStreamName.Replace(TEXT("/"), TEXT("+")) + TEXT("\""));
 					if(!UpdateVersionFile(*(ClientRootPath + VersionHeaderFileName), VersionHeaderStrings, PendingChangeNumber))
 					{
-						OutStatusMessage = FString::Printf(TEXT("Failed to update %s."), *VersionHeaderFileName);
+						OutStatusMessage = FString::Printf(TEXT("Failed to update %s."), VersionHeaderFileName);
 						return EWorkspaceUpdateResult::FailedToSync;
 					}
 					if(!UpdateVersionFile(*(ClientRootPath + ObjectVersionFileName), TMap<FString,FString>(), PendingChangeNumber))
 					{
-						OutStatusMessage = FString::Printf(TEXT("Failed to update %s."), *ObjectVersionFileName);
+						OutStatusMessage = FString::Printf(TEXT("Failed to update %s."), ObjectVersionFileName);
 						return EWorkspaceUpdateResult::FailedToSync;
 					}
 				}
@@ -600,7 +600,7 @@ EWorkspaceUpdateResult FWorkspace::UpdateWorkspaceInternal(FWorkspaceUpdateConte
 				{
 					if(!UpdateVersionFile(*(ClientRootPath + BuildVersionFileName), TMap<FString, FString>(), PendingChangeNumber))
 					{
-						OutStatusMessage = FString::Printf(TEXT("Failed to update %s"), *BuildVersionFileName);
+						OutStatusMessage = FString::Printf(TEXT("Failed to update %s"), BuildVersionFileName);
 						return EWorkspaceUpdateResult::FailedToSync;
 					}
 
@@ -611,12 +611,12 @@ EWorkspaceUpdateResult FWorkspace::UpdateWorkspaceInternal(FWorkspaceUpdateConte
 					VersionStrings.Add(TEXT("#define BRANCH_NAME"), FString::Printf(TEXT(" \"%s\""), *BranchOrStreamName.Replace(TEXT("/"), TEXT("+"))));
 					if(!UpdateVersionFile(*(ClientRootPath + VersionHeaderFileName), VersionStrings, PendingChangeNumber))
 					{
-						OutStatusMessage = FString::Printf(TEXT("Failed to update %s"), *VersionHeaderFileName);
+						OutStatusMessage = FString::Printf(TEXT("Failed to update %s"), VersionHeaderFileName);
 						return EWorkspaceUpdateResult::FailedToSync;
 					}
 					if(!UpdateVersionFile(*(ClientRootPath + ObjectVersionFileName), VersionStrings, PendingChangeNumber))
 					{
-						OutStatusMessage = FString::Printf(TEXT("Failed to update %s"), *ObjectVersionFileName);
+						OutStatusMessage = FString::Printf(TEXT("Failed to update %s"), ObjectVersionFileName);
 						return EWorkspaceUpdateResult::FailedToSync;
 					}
 				}
diff --git a/Engine/Source/Runtime/Core/Private/Modules/Tests/VisualizerDebuggingStateTests.cpp b/Engine/Source/Runtime/Core/Private/Modules/Tests/VisualizerDebuggingStateTests.cpp
--- a/Engine/Source/Runtime/Core/Private/Modules/Tests/VisualizerDebuggingStateTests.cpp
+++ b/Engine/Source/Runtime/Core/Private/Modules/Tests/VisualizerDebuggingStateTests.cpp
@@ -107,7 +107,7 @@ bool FCoreDebuggingStateTest::RunTest(const FString& Parameters)
 
 	if (TestState.Assign(CollisionId, nullptr) != EVisualizerDebuggingStateResult::StringCollision)
 	{
-		this->AddError(FString::Printf(TEXT("Expected FVisualizerDebuggingState::Assign(\"%s\") to result in a string collision but it did not"), CollisionGuid));
+		this->AddError(FString::Printf(TEXT("Expected FVisualizerDebuggingState::Assign(\"%hs\") to result in a string collision but it did not"), CollisionGuid));
 		return false;
 	}
 
diff --git a/Engine/Source/Runtime/Engine/Private/ShaderCompiler/ShaderCompiler.cpp b/Engine/Source/Runtime/Engine/Private/ShaderCompiler/ShaderCompiler.cpp
--- a/Engine/Source/Runtime/Engine/Private/ShaderCompiler/ShaderCompiler.cpp
+++ b/Engine/Source/Runtime/Engine/Private/ShaderCompiler/ShaderCompiler.cpp
@@ -3915,7 +3915,7 @@ void GlobalBeginCompileShader(
 		// We limit matrix definitions below to WORKING_COLOR_SPACE_IS_SRGB == 0.
 		if (!bWorkingColorSpaceIsSRGB)
 		{
-			const TCHAR MatrixFormat[] = TEXT("float3x3(%0.10f, %0.10f, %0.10f, %0.10f, %0.10f, %0.10f, %0.10f, %0.10f, %0.10f)");
+			static constexpr TCHAR MatrixFormat[] = TEXT("float3x3(%0.10f, %0.10f, %0.10f, %0.10f, %0.10f, %0.10f, %0.10f, %0.10f, %0.10f)");
 			const FColorSpace& WorkingColorSpace = FColorSpace::GetWorking();
 
 			// Note that we transpose the matrices during print since color matrices are usually pre-multiplied.
diff --git a/Engine/Source/Runtime/Engine/Private/ShaderCompiler/ShaderCompilerJobCache.cpp b/Engine/Source/Runtime/Engine/Private/ShaderCompiler/ShaderCompilerJobCache.cpp
--- a/Engine/Source/Runtime/Engine/Private/ShaderCompiler/ShaderCompilerJobCache.cpp
+++ b/Engine/Source/Runtime/Engine/Private/ShaderCompiler/ShaderCompilerJobCache.cpp
@@ -2231,7 +2231,7 @@ void FShaderJobCache::AddJobOutput(FShaderJobData& JobData, const FShaderCommonC
 				FFileHelper::SaveArrayToFile(TArrayView<const uint8>((const uint8*)SaveContext.ShaderObjectData.GetData(), SaveContext.ShaderObjectData.GetSize()), *CacheOutputFilename);
 				for (int32 CodeIndex = 0; CodeIndex < SaveContext.ShaderCode.Num(); ++CodeIndex)
 				{
-					const FString CacheCodeFilename = FString::Printf(TEXT("%s/%s_%d.bytecode"), *InputDebugInfoPath, *InputSourceFilename);
+					const FString CacheCodeFilename = FString::Printf(TEXT("%s/%s_%d.bytecode"), *InputDebugInfoPath, *InputSourceFilename, CodeIndex);
 					const FSharedBuffer& JobCode = SaveContext.ShaderCode[CodeIndex];
 					FFileHelper::SaveArrayToFile(TArrayView<const uint8>((const uint8*)JobCode.GetData(), JobCode.GetSize()), *CacheCodeFilename);
 				}
diff --git a/Engine/Source/Runtime/MovieScene/Private/EntitySystem/MovieSceneEntitySystem.cpp b/Engine/Source/Runtime/MovieScene/Private/EntitySystem/MovieSceneEntitySystem.cpp
--- a/Engine/Source/Runtime/MovieScene/Private/EntitySystem/MovieSceneEntitySystem.cpp
+++ b/Engine/Source/Runtime/MovieScene/Private/EntitySystem/MovieSceneEntitySystem.cpp
@@ -149,7 +149,7 @@ struct FSystemDependencyGraph
 
 	void DebugPrint(bool bUpdateCache)
 	{
-		const TCHAR FormatString[] =
+		static constexpr TCHAR FormatString[] =
 			TEXT("----------------------------------------------------------------------------------\n")
 			TEXT("%s\n")
 			TEXT("----------------------------------------------------------------------------------\n");
@@ -245,7 +245,7 @@ struct FSystemDependencyGraph
 #if UE_MOVIESCENE_ENTITY_DEBUG
 						FlowStrings[0] += FString::Printf(TEXT("\t\tflow_node%d_0[label=\"%s\"];\n"), NodeID, *WriteComponentInfo.DebugInfo->DebugName);
 #else
-						FlowStrings[0] += FString::Printf(TEXT("\t\tflow_node%d_0[label=\"Component ID=%s\"];\n"), NodeID, Node.WriteComponentType.BitIndex());
+						FlowStrings[0] += FString::Printf(TEXT("\t\tflow_node%d_0[label=\"Component ID=%d\"];\n"), NodeID, Node.WriteComponentType.BitIndex());
 #endif
 					}
 					if (EnumHasAnyFlags(*ConnectedSystemPhase, ESystemPhase::Instantiation))
@@ -253,7 +253,7 @@ struct FSystemDependencyGraph
 #if UE_MOVIESCENE_ENTITY_DEBUG
 						FlowStrings[1] += FString::Printf(TEXT("\t\tflow_node%d_1[label=\"%s\"];\n"), NodeID, *WriteComponentInfo.DebugInfo->DebugName);
 #else
-						FlowStrings[1] += FString::Printf(TEXT("\t\tflow_node%d_1[label=\"Component ID=%s\"];\n"), NodeID, Node.WriteComponentType.BitIndex());
+						FlowStrings[1] += FString::Printf(TEXT("\t\tflow_node%d_1[label=\"Component ID=%d\"];\n"), NodeID, Node.WriteComponentType.BitIndex());
 #endif
 					}
 					if (EnumHasAnyFlags(*ConnectedSystemPhase, ESystemPhase::Evaluation | ESystemPhase::Scheduling))
@@ -261,7 +261,7 @@ struct FSystemDependencyGraph
 #if UE_MOVIESCENE_ENTITY_DEBUG
 						FlowStrings[2] += FString::Printf(TEXT("\t\tflow_node%d_2[label=\"%s\"];\n"), NodeID, *WriteComponentInfo.DebugInfo->DebugName);
 #else
-						FlowStrings[2] += FString::Printf(TEXT("\t\tflow_node%d_2[label=\"Component ID=%s\"];\n"), NodeID, Node.WriteComponentType.BitIndex());
+						FlowStrings[2] += FString::Printf(TEXT("\t\tflow_node%d_2[label=\"Component ID=%d\"];\n"), NodeID, Node.WriteComponentType.BitIndex());
 #endif
 					}
 					if (EnumHasAnyFlags(*ConnectedSystemPhase, ESystemPhase::Finalization))
@@ -269,7 +269,7 @@ struct FSystemDependencyGraph
 #if UE_MOVIESCENE_ENTITY_DEBUG
 						FlowStrings[3] += FString::Printf(TEXT("\t\tflow_node%d_3[label=\"%s\"];\n"), NodeID, *WriteComponentInfo.DebugInfo->DebugName);
 #else
-						FlowStrings[3] += FString::Printf(TEXT("\t\tflow_node%d_3[label=\"Component ID=%s\"];\n"), NodeID, Node.WriteComponentType.BitIndex());
+						FlowStrings[3] += FString::Printf(TEXT("\t\tflow_node%d_3[label=\"Component ID=%d\"];\n"), NodeID, Node.WriteComponentType.BitIndex());
 #endif
 					}
 				}
diff --git a/Engine/Source/Runtime/MovieScene/Private/EntitySystem/MovieSceneEntitySystemGraphs.cpp b/Engine/Source/Runtime/MovieScene/Private/EntitySystem/MovieSceneEntitySystemGraphs.cpp
--- a/Engine/Source/Runtime/MovieScene/Private/EntitySystem/MovieSceneEntitySystemGraphs.cpp
+++ b/Engine/Source/Runtime/MovieScene/Private/EntitySystem/MovieSceneEntitySystemGraphs.cpp
@@ -225,7 +225,7 @@ void FMovieSceneEntitySystemGraph::UpdateCache()
 
 void FMovieSceneEntitySystemGraph::DebugPrint() const
 {
-	const TCHAR FormatString[] =
+	static constexpr TCHAR FormatString[] =
 		TEXT("----------------------------------------------------------------------------------\n")
 		TEXT("%s\n")
 		TEXT("----------------------------------------------------------------------------------\n");
diff --git a/Engine/Source/Runtime/RHI/Private/PipelineFileCache.cpp b/Engine/Source/Runtime/RHI/Private/PipelineFileCache.cpp
--- a/Engine/Source/Runtime/RHI/Private/PipelineFileCache.cpp
+++ b/Engine/Source/Runtime/RHI/Private/PipelineFileCache.cpp
@@ -4367,7 +4367,7 @@ FString FPipelineCacheFileFormatPSO::FPipelineFileCacheRayTracingDesc::HeaderLin
 
 FString FPipelineCacheFileFormatPSO::FPipelineFileCacheRayTracingDesc::ToString() const
 {
-	return FString::Printf(TEXT("%s,%d,%d,%d")
+	return FString::Printf(TEXT("%s,%d,%d")
 		, *ShaderHash.ToString()
 		, DeprecatedMaxPayloadSizeInBytes
 		, uint32(Frequency)
diff --git a/Engine/Source/Runtime/RHICore/Private/RHICoreNvidiaAftermath.cpp b/Engine/Source/Runtime/RHICore/Private/RHICoreNvidiaAftermath.cpp
--- a/Engine/Source/Runtime/RHICore/Private/RHICoreNvidiaAftermath.cpp
+++ b/Engine/Source/Runtime/RHICore/Private/RHICoreNvidiaAftermath.cpp
@@ -412,7 +412,7 @@ namespace UE::RHICore::Nvidia::Aftermath
 						const GFSDK_Aftermath_GpuCrashDump_ResourceInfo& FaultResInfo = FaultResources[ResIdx];
 						CrashResult.OutputLog += FString::Printf(
 							TEXT("\n\t\tResource %d/%d:")
-							TEXT("\n\t\t\tName                : '%s'")
+							TEXT("\n\t\t\tName                : '%hs'")
 							TEXT("\n\t\t\tGPU VA              : 0x%016llx")
 							TEXT("\n\t\t\tSize                : 0x%016llx")
 							TEXT("\n\t\t\tWidth               : %d")
@@ -563,7 +563,7 @@ namespace UE::RHICore::Nvidia::Aftermath
 			}
 			else
 			{
-				CrashResult.OutputLog += FString::Printf(TEXT("\n\n%hs\n\n"), JsonStr.GetData());
+				CrashResult.OutputLog += FString::Printf(TEXT("\n\n%s\n\n"), StringCast<TCHAR>(JsonStr.GetData()).Get());
 			}
 		}
 	}
diff --git a/Engine/Source/Runtime/Renderer/Private/PostProcess/PostProcessBokehDOF.cpp b/Engine/Source/Runtime/Renderer/Private/PostProcess/PostProcessBokehDOF.cpp
--- a/Engine/Source/Runtime/Renderer/Private/PostProcess/PostProcessBokehDOF.cpp
+++ b/Engine/Source/Runtime/Renderer/Private/PostProcess/PostProcessBokehDOF.cpp
@@ -109,7 +109,7 @@ FScreenPassTexture AddVisualizeDOFPass(FRDGBuilder& GraphBuilder, const FViewInf
 		Y += YStep;
 		Line = FString::Printf(TEXT("DepthBlur (not related to Depth of Field, due to light traveling long distances in atmosphere)"));
 		Canvas.DrawShadowedString(X, Y += YStep, *Line, GetStatsFont(), FLinearColor(1, 1, 1));
-		Line = FString::Printf(TEXT("  km for 50%: %.2f"), View.FinalPostProcessSettings.DepthOfFieldDepthBlurAmount);
+		Line = FString::Printf(TEXT("  km for 50%%: %.2f"), View.FinalPostProcessSettings.DepthOfFieldDepthBlurAmount);
 		Canvas.DrawShadowedString(X, Y += YStep, *Line, GetStatsFont(), FLinearColor(1, 1, 1));
 		Line = FString::Printf(TEXT("  Radius (pixels in 1920x): %.2f"), View.FinalPostProcessSettings.DepthOfFieldDepthBlurRadius);
 		Canvas.DrawShadowedString(X, Y += YStep, *Line, GetStatsFont(), FLinearColor(1, 1, 1));
diff --git a/Engine/Source/Runtime/Renderer/Private/VolumetricCloudRendering.cpp b/Engine/Source/Runtime/Renderer/Private/VolumetricCloudRendering.cpp
--- a/Engine/Source/Runtime/Renderer/Private/VolumetricCloudRendering.cpp
+++ b/Engine/Source/Runtime/Renderer/Private/VolumetricCloudRendering.cpp
@@ -3008,7 +3008,7 @@ bool FSceneRenderer::RenderVolumetricCloud(
 							const float ViewPortWidth = float(ViewInfo.ViewRect.Width());
 							const float ViewPortHeight = float(ViewInfo.ViewRect.Height());
 							FLinearColor TextColor(1.0f, 0.5f, 0.0f);
-							FString Text = FString::Printf(TEXT("Shadow Sample Count = %.1f"), CloudShadowmapSampleCount);
+							FString Text = FString::Printf(TEXT("Shadow Sample Count = %.1f"), float(CloudShadowmapSampleCount));
 							Canvas.DrawShadowedString(0.05f, ViewPortHeight * 0.4f, *Text, GetStatsFont(), TextColor);
 						});
 
