Slate UGS: Fix client filter

#rnx
#Synced-CL 51838013

diff --git a/Engine/Source/Programs/UGSCore/DetectProjectSettingsTask.cpp b/Engine/Source/Programs/UGSCore/DetectProjectSettingsTask.cpp
--- a/Engine/Source/Programs/UGSCore/DetectProjectSettingsTask.cpp
+++ b/Engine/Source/Programs/UGSCore/DetectProjectSettingsTask.cpp
@@ -37,14 +37,8 @@ TSharedRef<FModalTaskResult> FDetectProjectSettingsTask::Run(FEvent* AbortEvent)
 
 TSharedRef<FModalTaskResult> FDetectProjectSettingsTask::RunInternal(FEvent* AbortEvent)
 {
-	// Get the P4PORT setting in this folder, so we can respect the contents of any P4CONFIG file
 	FString PrevDirectory = FPlatformProcess::GetCurrentWorkingDirectory();
 
-//	Directory.SetCurrentDirectory(Path.GetDirectoryName(NewSelectedFileName));
-//
-	FString ServerAndPort;
-	Perforce->GetSetting(TEXT("P4PORT"), ServerAndPort, AbortEvent, Log.Get());
-
 	// Get the Perforce server info
 	TSharedPtr<FPerforceInfoRecord> PerforceInfo;
 	if(!Perforce->Info(PerforceInfo, AbortEvent, Log.Get()))
@@ -61,11 +55,6 @@ TSharedRef<FModalTaskResult> FDetectProjectSettingsTask::RunInternal(FEvent* Abo
 	}
 	ServerTimeZone = PerforceInfo->ServerTimeZone;
 
-	if (ServerAndPort.IsEmpty())
-	{
-		ServerAndPort = PerforceInfo->ServerAddress;
-	}
-
 	// Find all the clients on this machine
 	Log->Logf(TEXT("Enumerating clients on local machine..."));
 	TArray<FPerforceClientRecord> Clients;
@@ -90,10 +79,10 @@ TSharedRef<FModalTaskResult> FDetectProjectSettingsTask::RunInternal(FEvent* Abo
 					continue;
 				}
 
-				TSharedRef<FPerforceConnection> CandidateClient = MakeShared<FPerforceConnection>(*PerforceInfo->UserName, *Client.Name, *ServerAndPort);
+				TSharedRef<FPerforceConnection> CandidateClient = Perforce->OpenClient(Client.Name);
 
-				bool bFileExists;
-				if(!CandidateClient->FileExists(NewSelectedFileName, bFileExists, AbortEvent, Log.Get()) || !bFileExists)
+				TSharedPtr<FPerforceWhereRecord> WhereRecord;
+				if (!CandidateClient->Where(NewSelectedFileName, WhereRecord, AbortEvent, Log.Get()))
 				{
 					Log->Logf(TEXT("Rejecting %s due to file not existing in workspace"), *Client.Name);
 					continue;
