Slate UGS: Fix global filter

#rnx
#Synced-CL 51838013

diff --git a/Engine/Source/Programs/SlateUGS/Private/UGSTab.cpp b/Engine/Source/Programs/SlateUGS/Private/UGSTab.cpp
--- a/Engine/Source/Programs/SlateUGS/Private/UGSTab.cpp
+++ b/Engine/Source/Programs/SlateUGS/Private/UGSTab.cpp
@@ -364,14 +364,18 @@ void UGSTab::OnSyncLatest()
 
 void UGSTab::OnSyncFilterWindowSaved(
 	const TArray<FString>& SyncViewCurrent,
+	const TArray<FGuid>& SyncIncludedCategoriesCurrent,
 	const TArray<FGuid>& SyncExcludedCategoriesCurrent,
 	const TArray<FString>& SyncViewAll,
+	const TArray<FGuid>& SyncIncludedCategoriesAll,
 	const TArray<FGuid>& SyncExcludedCategoriesAll)
 {
 	WorkspaceSettings->SyncView = SyncViewCurrent;
+	WorkspaceSettings->SyncIncludedCategories = SyncIncludedCategoriesCurrent;
 	WorkspaceSettings->SyncExcludedCategories = SyncExcludedCategoriesCurrent;
 
 	UserSettings->SyncView = SyncViewAll;
+	UserSettings->SyncIncludedCategories = SyncIncludedCategoriesAll;
 	UserSettings->SyncExcludedCategories = SyncExcludedCategoriesAll;
 
 	UserSettings->Save();
@@ -616,20 +620,34 @@ TArray<UGSCore::FWorkspaceSyncCategory> UGSTab::GetSyncCategories(SyncCategoryTy
 {
 	TArray<UGSCore::FWorkspaceSyncCategory> SyncCategories;
 	TMap<FGuid, UGSCore::FWorkspaceSyncCategory> Categories = Workspace->GetSyncCategories();
+	TArray<FGuid> SyncIncludedCategories;
 	TArray<FGuid> SyncExcludedCategories;
 
 	if (CategoryType == SyncCategoryType::CurrentWorkspace)
 	{
+		SyncIncludedCategories = WorkspaceSettings->SyncIncludedCategories;
 		SyncExcludedCategories = WorkspaceSettings->SyncExcludedCategories;
 	}
 	else
 	{
+		SyncIncludedCategories = UserSettings->SyncIncludedCategories;
 		SyncExcludedCategories = UserSettings->SyncExcludedCategories;
 	}
 
+	for (const FGuid& Guid : SyncIncludedCategories)
+	{
+		if (UGSCore::FWorkspaceSyncCategory* Category = Categories.Find(Guid))
+		{
+			Category->bEnable = true;
+		}
+	}
+
 	for (const FGuid& Guid : SyncExcludedCategories)
 	{
-		Categories[Guid].bEnable = false;
+		if (UGSCore::FWorkspaceSyncCategory* Category = Categories.Find(Guid))
+		{
+			Category->bEnable = false;
+		}
 	}
 
 	Categories.GenerateValueArray(SyncCategories);
diff --git a/Engine/Source/Programs/SlateUGS/Private/UGSTab.h b/Engine/Source/Programs/SlateUGS/Private/UGSTab.h
--- a/Engine/Source/Programs/SlateUGS/Private/UGSTab.h
+++ b/Engine/Source/Programs/SlateUGS/Private/UGSTab.h
@@ -56,8 +56,10 @@ public:
 	void OnSyncLatest();
 	void OnSyncFilterWindowSaved(
 		const TArray<FString>& SyncViewCurrent,
+		const TArray<FGuid>& SyncIncludedCategoriesCurrent,
 		const TArray<FGuid>& SyncExcludedCategoriesCurrent,
 		const TArray<FString>& SyncViewAll,
+		const TArray<FGuid>& SyncIncludedCategoriesAll,
 		const TArray<FGuid>& SyncExcludedCategoriesAll);
 	void OnBuildWorkspace();
 	void OnOpenExplorer();
diff --git a/Engine/Source/Programs/SlateUGS/Private/Widgets/SSyncFilterWindow.cpp b/Engine/Source/Programs/SlateUGS/Private/Widgets/SSyncFilterWindow.cpp
--- a/Engine/Source/Programs/SlateUGS/Private/Widgets/SSyncFilterWindow.cpp
+++ b/Engine/Source/Programs/SlateUGS/Private/Widgets/SSyncFilterWindow.cpp
@@ -265,6 +265,9 @@ FReply SSyncFilterWindow::OnCustomViewSyntaxClicked()
 
 FReply SSyncFilterWindow::OnSaveClicked()
 {
+	// Only store the delta from defaults, (all built-in categories default to enabled):
+	//   checked   → matches default → not stored
+	//   unchecked → differs from default → stored in ExcludeCategories
 	TArray<FGuid> SyncExcludedCategoriesCurrent;
 	TArray<bool> EnabledValuesCurrent = SyncFiltersCurrent->GetValues();
 	for (int CategoryIndex = 0; CategoryIndex < WorkspaceCategoriesCurrent.Num(); CategoryIndex++)
@@ -291,7 +294,8 @@ FReply SSyncFilterWindow::OnSaveClicked()
 	TArray<FString> SyncViewAll;
 	CustomSyncViewAll->GetPlainText().ToString().ParseIntoArray(SyncViewAll, TEXT("\n"));
 
-	Tab->OnSyncFilterWindowSaved(SyncViewCurrent, SyncExcludedCategoriesCurrent, SyncViewAll, SyncExcludedCategoriesAll);
+	TArray<FGuid> EmptyIncludeCategories;
+	Tab->OnSyncFilterWindowSaved(SyncViewCurrent, EmptyIncludeCategories, SyncExcludedCategoriesCurrent, SyncViewAll, EmptyIncludeCategories, SyncExcludedCategoriesAll);
 
 	FSlateApplication::Get().FindWidgetWindow(AsShared())->RequestDestroyWindow();
 	return FReply::Handled();
diff --git a/Engine/Source/Programs/UGSCore/UserSettings.cpp b/Engine/Source/Programs/UGSCore/UserSettings.cpp
--- a/Engine/Source/Programs/UGSCore/UserSettings.cpp
+++ b/Engine/Source/Programs/UGSCore/UserSettings.cpp
@@ -4,6 +4,13 @@
 #include "HAL/FileManager.h"
 #include "Algo/Transform.h"
 #include "Utility.h"
+#include "Dom/JsonObject.h"
+#include "Dom/JsonValue.h"
+#include "Misc/FileHelper.h"
+#include "Misc/Paths.h"
+#include "Serialization/JsonReader.h"
+#include "Serialization/JsonSerializer.h"
+#include "Serialization/JsonWriter.h"
 
 namespace UGSCore
 {
@@ -172,6 +179,7 @@ FUserSettings::FUserSettings(const FString& InFileName)
 
 	ConfigFile.TryGetValues(TEXT("General.OtherProjectFileNames"), OtherProjectFileNames);
 	ConfigFile.TryGetValues(TEXT("General.SyncFilter"), SyncView);
+	ConfigFile.TryGetValues(TEXT("General.SyncIncludedCategories"), SyncIncludedCategories);
 	ConfigFile.TryGetValues(TEXT("General.SyncExcludedCategories"), SyncExcludedCategories);
 
 	// TODO i set this default value as we assert when saving if no value is set but honestly idk what Any means
@@ -258,6 +266,8 @@ FUserSettings::FUserSettings(const FString& InFileName)
 	SyncOptions.NumRetries = ConfigFile.GetValue(TEXT("Perforce.NumRetries"), 0);
 	SyncOptions.NumThreads = ConfigFile.GetValue(TEXT("Perforce.NumThreads"), 0);
 	SyncOptions.TcpBufferSize = ConfigFile.GetValue(TEXT("Perforce.TcpBufferSize"), 0);
+
+	LoadGlobalSettings();
 }
 
 FString Trim(const FString& Text, const TCHAR Character)
@@ -334,6 +344,7 @@ TSharedRef<FUserWorkspaceSettings> FUserSettings::FindOrAddWorkspace(const TCHAR
 			}
 
 			CurrentWorkspace->SyncView.Empty();
+			CurrentWorkspace->SyncIncludedCategories.Empty();
 			CurrentWorkspace->SyncExcludedCategories.Empty();
 		}
 		else
@@ -359,6 +370,7 @@ TSharedRef<FUserWorkspaceSettings> FUserSettings::FindOrAddWorkspace(const TCHAR
 			WorkspaceSection->TryGetValues(TEXT("ExpandedArchiveName"), CurrentWorkspace->ExpandedArchiveTypes);
 
 			WorkspaceSection->TryGetValues(TEXT("SyncFilter"), CurrentWorkspace->SyncView);
+			WorkspaceSection->TryGetValues(TEXT("SyncIncludedCategories"), CurrentWorkspace->SyncIncludedCategories);
 			WorkspaceSection->TryGetValues(TEXT("SyncExcludedCategories"), CurrentWorkspace->SyncExcludedCategories);
 		}
 	}
@@ -402,6 +414,7 @@ void FUserSettings::Save()
 	GeneralSection->SetValue(TEXT("FilterIndex"), FilterIndex);
 	GeneralSection->SetValues(TEXT("OtherProjectFileNames"), OtherProjectFileNames);
 	GeneralSection->SetValues(TEXT("SyncFilter"), SyncView);
+	GeneralSection->SetValues(TEXT("SyncIncludedCategories"), SyncIncludedCategories);
 	GeneralSection->SetValues(TEXT("SyncExcludedCategories"), SyncExcludedCategories);
 	GeneralSection->SetValue(TEXT("SyncType"), *ToString(SyncType));
 
@@ -471,6 +484,7 @@ void FUserSettings::Save()
 		WorkspaceSection->SetValue(TEXT("LastBuiltChangeNumber"), CurrentWorkspace.LastBuiltChangeNumber);
 		WorkspaceSection->SetValues(TEXT("ExpandedArchiveName"), CurrentWorkspace.ExpandedArchiveTypes);
 		WorkspaceSection->SetValues(TEXT("SyncFilter"), CurrentWorkspace.SyncView);
+		WorkspaceSection->SetValues(TEXT("SyncIncludedCategories"), CurrentWorkspace.SyncIncludedCategories);
 		WorkspaceSection->SetValues(TEXT("SyncExcludedCategories"), CurrentWorkspace.SyncExcludedCategories);
 	}
 
@@ -506,6 +520,8 @@ void FUserSettings::Save()
 
 	// Save the file
 	ConfigFile.Save(*FileName);
+
+	SaveGlobalSettings();
 }
 
 TArray<FString> FUserSettings::GetCombinedSyncFilter(const TMap<FGuid, FWorkspaceSyncCategory>& UniqueIdToFilter, const TArray<FString>& GlobalView, const TArray<FGuid>& GlobalExcludedCategories, const TArray<FString>& WorkspaceView, const TArray<FGuid>& WorkspaceExcludedCategories)
@@ -546,6 +562,124 @@ TArray<FString> FUserSettings::GetCombinedSyncFilter(const TMap<FGuid, FWorkspac
 	return Lines;
 }
 
+static FString GetGlobalJsonPath()
+{
+	FPlatformProcess::ApplicationSettingsContext Context;
+	Context.Location = FPlatformProcess::ApplicationSettingsContext::Context::LocalUser;
+	Context.bIsEpic = false;
+	return FPlatformProcess::GetApplicationSettingsDir(Context) / TEXT("UnrealGameSync") / TEXT("Global.json");
+}
+
+void FUserSettings::LoadGlobalSettings()
+{
+	FString GlobalJsonPath = GetGlobalJsonPath();
+
+	FString JsonString;
+	if (!FFileHelper::LoadFileToString(JsonString, *GlobalJsonPath))
+	{
+		return;
+	}
+
+	TSharedPtr<FJsonObject> RootObject;
+	TSharedRef<TJsonReader<>> Reader = TJsonReaderFactory<>::Create(JsonString);
+	if (!FJsonSerializer::Deserialize(Reader, RootObject) || !RootObject.IsValid())
+	{
+		return;
+	}
+
+	const TSharedPtr<FJsonObject>* FilterObject;
+	if (!RootObject->TryGetObjectField(TEXT("Filter"), FilterObject))
+	{
+		return;
+	}
+
+	const TArray<TSharedPtr<FJsonValue>>* ExcludeArray;
+	if ((*FilterObject)->TryGetArrayField(TEXT("ExcludeCategories"), ExcludeArray))
+	{
+		SyncExcludedCategories.Empty();
+		for (const TSharedPtr<FJsonValue>& Value : *ExcludeArray)
+		{
+			FGuid Guid;
+			if (FGuid::Parse(Value->AsString(), Guid))
+			{
+				SyncExcludedCategories.Add(Guid);
+			}
+		}
+	}
+
+	const TArray<TSharedPtr<FJsonValue>>* IncludeArray;
+	if ((*FilterObject)->TryGetArrayField(TEXT("IncludeCategories"), IncludeArray))
+	{
+		SyncIncludedCategories.Empty();
+		for (const TSharedPtr<FJsonValue>& Value : *IncludeArray)
+		{
+			FGuid Guid;
+			if (FGuid::Parse(Value->AsString(), Guid))
+			{
+				SyncIncludedCategories.Add(Guid);
+			}
+		}
+	}
+
+	const TArray<TSharedPtr<FJsonValue>>* ViewArray;
+	if ((*FilterObject)->TryGetArrayField(TEXT("View"), ViewArray))
+	{
+		SyncView.Empty();
+		for (const TSharedPtr<FJsonValue>& Value : *ViewArray)
+		{
+			SyncView.Add(Value->AsString());
+		}
+	}
+}
+
+void FUserSettings::SaveGlobalSettings()
+{
+	FString GlobalJsonPath = GetGlobalJsonPath();
+
+	// Read-modify-write to preserve other fields in Global.json
+	TSharedPtr<FJsonObject> RootObject;
+	FString JsonString;
+	if (FFileHelper::LoadFileToString(JsonString, *GlobalJsonPath))
+	{
+		TSharedRef<TJsonReader<>> Reader = TJsonReaderFactory<>::Create(JsonString);
+		FJsonSerializer::Deserialize(Reader, RootObject);
+	}
+	if (!RootObject.IsValid())
+	{
+		RootObject = MakeShared<FJsonObject>();
+	}
+
+	TSharedPtr<FJsonObject> FilterObject = MakeShared<FJsonObject>();
+
+	TArray<TSharedPtr<FJsonValue>> IncludeValues;
+	for (const FGuid& Guid : SyncIncludedCategories)
+	{
+		IncludeValues.Add(MakeShared<FJsonValueString>(Guid.ToString(EGuidFormats::DigitsWithHyphensLower)));
+	}
+	FilterObject->SetArrayField(TEXT("IncludeCategories"), IncludeValues);
+
+	TArray<TSharedPtr<FJsonValue>> ExcludeValues;
+	for (const FGuid& Guid : SyncExcludedCategories)
+	{
+		ExcludeValues.Add(MakeShared<FJsonValueString>(Guid.ToString(EGuidFormats::DigitsWithHyphensLower)));
+	}
+	FilterObject->SetArrayField(TEXT("ExcludeCategories"), ExcludeValues);
+
+	TArray<TSharedPtr<FJsonValue>> ViewValues;
+	for (const FString& ViewLine : SyncView)
+	{
+		ViewValues.Add(MakeShared<FJsonValueString>(ViewLine));
+	}
+	FilterObject->SetArrayField(TEXT("View"), ViewValues);
+
+	RootObject->SetObjectField(TEXT("Filter"), FilterObject);
+
+	FString OutputString;
+	TSharedRef<TJsonWriter<>> Writer = TJsonWriterFactory<>::Create(&OutputString);
+	FJsonSerializer::Serialize(RootObject.ToSharedRef(), Writer);
+	FFileHelper::SaveStringToFile(OutputString, *GlobalJsonPath);
+}
+
 FString FUserSettings::EscapeText(const FString& Text)
 {
 	FString Result;
diff --git a/Engine/Source/Programs/UGSCore/UserSettings.h b/Engine/Source/Programs/UGSCore/UserSettings.h
--- a/Engine/Source/Programs/UGSCore/UserSettings.h
+++ b/Engine/Source/Programs/UGSCore/UserSettings.h
@@ -64,6 +64,7 @@ struct FUserWorkspaceSettings
 
 	// Workspace specific SyncFilters
 	TArray<FString> SyncView;
+	TArray<FGuid> SyncIncludedCategories;
 	TArray<FGuid> SyncExcludedCategories;
 };
 
@@ -93,6 +94,7 @@ struct FUserSettings
 	TArray<FString> OpenProjectFileNames;
 	TArray<FString> OtherProjectFileNames;
 	TArray<FString> SyncView;
+	TArray<FGuid> SyncIncludedCategories;
 	TArray<FGuid> SyncExcludedCategories;
 	ELatestChangeType SyncType;
 	EBuildConfig CompiledEditorBuildConfig; // NB: This assumes not using precompiled editor. See CurrentBuildConfig.
@@ -123,11 +125,15 @@ struct FUserSettings
 	TSharedRef<FUserWorkspaceSettings> FindOrAddWorkspace(const TCHAR* ClientBranchPath);
 	TSharedRef<FUserProjectSettings> FindOrAddProject(const TCHAR* ClientProjectFileName);
 	void Save();
-	
+
 	static TArray<FString> GetCombinedSyncFilter(const TMap<FGuid, FWorkspaceSyncCategory>& UniqueIdToFilter, const TArray<FString>& GlobalView, const TArray<FGuid>& GlobalExcludedCategories, const TArray<FString>& WorkspaceView, const TArray<FGuid>& WorkspaceExcludedCategories);
 
 	static FString EscapeText(const FString& Text);
 	static FString UnescapeText(const FString& Text);
+
+private:
+	void LoadGlobalSettings();
+	void SaveGlobalSettings();
 };
 
 } // namespace UGSCore
